Renamed TypedElementInterfaces to TypedElementRuntime and moved UTypedElementSelectionSet to it

This lets UTypedElementSelectionSet (which has minimal dependencies) to be used by the Engine and InteractiveToolsFramework modules

diff --git a/Engine/Source/Editor/AssetPlacementEdMode/AssetPlacementEdMode.Build.cs b/Engine/Source/Editor/AssetPlacementEdMode/AssetPlacementEdMode.Build.cs
--- a/Engine/Source/Editor/AssetPlacementEdMode/AssetPlacementEdMode.Build.cs
+++ b/Engine/Source/Editor/AssetPlacementEdMode/AssetPlacementEdMode.Build.cs
@@ -22,7 +22,7 @@ public class AssetPlacementEdMode : ModuleRules
 				"Foliage",
 				"Landscape",
 				"TypedElementFramework",
-				"TypedElementInterfaces",
+				"TypedElementRuntime",
 				"PropertyEditor",
 			}
 		);
diff --git a/Engine/Source/Editor/EditorFramework/EditorFramework.Build.cs b/Engine/Source/Editor/EditorFramework/EditorFramework.Build.cs
--- a/Engine/Source/Editor/EditorFramework/EditorFramework.Build.cs
+++ b/Engine/Source/Editor/EditorFramework/EditorFramework.Build.cs
@@ -34,7 +34,7 @@ public class EditorFramework : ModuleRules
 				"EditorSubsystem",
 				"InteractiveToolsFramework",
 				"TypedElementFramework",
-				"TypedElementInterfaces",
+				"TypedElementRuntime",
 			}
 		);
 
diff --git a/Engine/Source/Editor/EditorFramework/Private/Elements/Framework/TypedElementAssetEditorCustomization.cpp b/Engine/Source/Editor/EditorFramework/Private/Elements/Framework/TypedElementAssetEditorCustomization.cpp
deleted file mode 100644
diff --git a/Engine/Source/Editor/EditorFramework/Private/Elements/Framework/TypedElementListObjectUtil.cpp b/Engine/Source/Editor/EditorFramework/Private/Elements/Framework/TypedElementListObjectUtil.cpp
deleted file mode 100644
diff --git a/Engine/Source/Editor/EditorFramework/Private/Elements/Framework/TypedElementSelectionSet.cpp b/Engine/Source/Editor/EditorFramework/Private/Elements/Framework/TypedElementSelectionSet.cpp
deleted file mode 100644
diff --git a/Engine/Source/Editor/EditorFramework/Private/Elements/Framework/TypedElementViewportInteraction.cpp b/Engine/Source/Editor/EditorFramework/Private/Elements/Framework/TypedElementViewportInteraction.cpp
--- a/Engine/Source/Editor/EditorFramework/Private/Elements/Framework/TypedElementViewportInteraction.cpp
+++ b/Engine/Source/Editor/EditorFramework/Private/Elements/Framework/TypedElementViewportInteraction.cpp
@@ -5,12 +5,12 @@
 #include "Elements/Framework/TypedElementRegistry.h"
 #include "Elements/Framework/TypedElementUtil.h"
 
-void FTypedElementAssetEditorViewportInteractionCustomization::GetElementsToMove(const TTypedElement<UTypedElementWorldInterface>& InElementWorldHandle, const ETypedElementViewportInteractionWorldType InWorldType, const UTypedElementSelectionSet* InSelectionSet, UTypedElementList* OutElementsToMove)
+void FTypedElementViewportInteractionCustomization::GetElementsToMove(const TTypedElement<UTypedElementWorldInterface>& InElementWorldHandle, const ETypedElementViewportInteractionWorldType InWorldType, const UTypedElementSelectionSet* InSelectionSet, UTypedElementList* OutElementsToMove)
 {
 	OutElementsToMove->Add(InElementWorldHandle);
 }
 
-bool FTypedElementAssetEditorViewportInteractionCustomization::GetGizmoPivotLocation(const TTypedElement<UTypedElementWorldInterface>& InElementWorldHandle, const UE::Widget::EWidgetMode InWidgetMode, FVector& OutPivotLocation)
+bool FTypedElementViewportInteractionCustomization::GetGizmoPivotLocation(const TTypedElement<UTypedElementWorldInterface>& InElementWorldHandle, const UE::Widget::EWidgetMode InWidgetMode, FVector& OutPivotLocation)
 {
 	FTransform ElementWorldTransform;
 	if (InElementWorldHandle.GetWorldTransform(ElementWorldTransform))
@@ -22,16 +22,16 @@ bool FTypedElementAssetEditorViewportInteractionCustomization::GetGizmoPivotLoca
 	return false;
 }
 
-void FTypedElementAssetEditorViewportInteractionCustomization::PreGizmoManipulationStarted(TArrayView<const FTypedElementHandle> InElementHandles, const UE::Widget::EWidgetMode InWidgetMode)
+void FTypedElementViewportInteractionCustomization::PreGizmoManipulationStarted(TArrayView<const FTypedElementHandle> InElementHandles, const UE::Widget::EWidgetMode InWidgetMode)
 {
 }
 
-void FTypedElementAssetEditorViewportInteractionCustomization::GizmoManipulationStarted(const TTypedElement<UTypedElementWorldInterface>& InElementWorldHandle, const UE::Widget::EWidgetMode InWidgetMode)
+void FTypedElementViewportInteractionCustomization::GizmoManipulationStarted(const TTypedElement<UTypedElementWorldInterface>& InElementWorldHandle, const UE::Widget::EWidgetMode InWidgetMode)
 {
 	InElementWorldHandle.NotifyMovementStarted();
 }
 
-void FTypedElementAssetEditorViewportInteractionCustomization::GizmoManipulationDeltaUpdate(const TTypedElement<UTypedElementWorldInterface>& InElementWorldHandle, const UE::Widget::EWidgetMode InWidgetMode, const EAxisList::Type InDragAxis, const FInputDeviceState& InInputState, const FTransform& InDeltaTransform, const FVector& InPivotLocation)
+void FTypedElementViewportInteractionCustomization::GizmoManipulationDeltaUpdate(const TTypedElement<UTypedElementWorldInterface>& InElementWorldHandle, const UE::Widget::EWidgetMode InWidgetMode, const EAxisList::Type InDragAxis, const FInputDeviceState& InInputState, const FTransform& InDeltaTransform, const FVector& InPivotLocation)
 {
 	FTransform ElementWorldTransform;
 	if (InElementWorldHandle.GetWorldTransform(ElementWorldTransform))
@@ -77,16 +77,16 @@ void FTypedElementAssetEditorViewportInteractionCustomization::GizmoManipulation
 	}
 }
 
-void FTypedElementAssetEditorViewportInteractionCustomization::GizmoManipulationStopped(const TTypedElement<UTypedElementWorldInterface>& InElementWorldHandle, const UE::Widget::EWidgetMode InWidgetMode)
+void FTypedElementViewportInteractionCustomization::GizmoManipulationStopped(const TTypedElement<UTypedElementWorldInterface>& InElementWorldHandle, const UE::Widget::EWidgetMode InWidgetMode)
 {
 	InElementWorldHandle.NotifyMovementEnded();
 }
 
-void FTypedElementAssetEditorViewportInteractionCustomization::PostGizmoManipulationStopped(TArrayView<const FTypedElementHandle> InElementHandles, const UE::Widget::EWidgetMode InWidgetMode)
+void FTypedElementViewportInteractionCustomization::PostGizmoManipulationStopped(TArrayView<const FTypedElementHandle> InElementHandles, const UE::Widget::EWidgetMode InWidgetMode)
 {
 }
 
-void FTypedElementAssetEditorViewportInteractionCustomization::MirrorElement(const TTypedElement<UTypedElementWorldInterface>& InElementWorldHandle, const FVector& InMirrorScale, const FVector& InPivotLocation)
+void FTypedElementViewportInteractionCustomization::MirrorElement(const TTypedElement<UTypedElementWorldInterface>& InElementWorldHandle, const FVector& InMirrorScale, const FVector& InPivotLocation)
 {
 	FTransform ElementWorldTransform;
 	if (InElementWorldHandle.GetWorldTransform(ElementWorldTransform))
@@ -143,7 +143,7 @@ void UTypedElementViewportInteraction::GetSelectedElementsToMove(const UTypedEle
 	{
 		if (InElementWorldHandle.CanEditElement())
 		{
-			FTypedElementViewportInteractionElement ViewportInteractionElement(InElementWorldHandle, GetAssetEditorCustomizationByTypeId(InElementWorldHandle.GetId().GetTypeId()));
+			FTypedElementViewportInteractionElement ViewportInteractionElement(InElementWorldHandle, GetInterfaceCustomizationByTypeId(InElementWorldHandle.GetId().GetTypeId()));
 			check(ViewportInteractionElement.IsSet());
 			ViewportInteractionElement.GetElementsToMove(InWorldType, InSelectionSet, OutElementsToMove);
 		}
@@ -159,9 +159,9 @@ void UTypedElementViewportInteraction::BeginGizmoManipulation(const UTypedElemen
 
 		for (const auto& ElementsByTypePair : ElementsToMoveByType)
 		{
-			FTypedElementAssetEditorViewportInteractionCustomization* AssetEditorViewportInteractionCustomization = GetAssetEditorCustomizationByTypeId(ElementsByTypePair.Key);
-			check(AssetEditorViewportInteractionCustomization);
-			AssetEditorViewportInteractionCustomization->PreGizmoManipulationStarted(ElementsByTypePair.Value, InWidgetMode);
+			FTypedElementViewportInteractionCustomization* ViewportInteractionCustomization = GetInterfaceCustomizationByTypeId(ElementsByTypePair.Key);
+			check(ViewportInteractionCustomization);
+			ViewportInteractionCustomization->PreGizmoManipulationStarted(ElementsByTypePair.Value, InWidgetMode);
 		}
 	}
 
@@ -209,9 +209,9 @@ void UTypedElementViewportInteraction::EndGizmoManipulation(const UTypedElementL
 
 		for (const auto& ElementsByTypePair : ElementsToMoveByType)
 		{
-			FTypedElementAssetEditorViewportInteractionCustomization* AssetEditorViewportInteractionCustomization = GetAssetEditorCustomizationByTypeId(ElementsByTypePair.Key);
-			check(AssetEditorViewportInteractionCustomization);
-			AssetEditorViewportInteractionCustomization->PostGizmoManipulationStopped(ElementsByTypePair.Value, InWidgetMode);
+			FTypedElementViewportInteractionCustomization* ViewportInteractionCustomization = GetInterfaceCustomizationByTypeId(ElementsByTypePair.Key);
+			check(ViewportInteractionCustomization);
+			ViewportInteractionCustomization->PostGizmoManipulationStopped(ElementsByTypePair.Value, InWidgetMode);
 		}
 	}
 }
@@ -241,6 +241,6 @@ void UTypedElementViewportInteraction::MirrorElement(const FTypedElementHandle&
 FTypedElementViewportInteractionElement UTypedElementViewportInteraction::ResolveViewportInteractionElement(const FTypedElementHandle& InElementHandle) const
 {
 	return InElementHandle
-		? FTypedElementViewportInteractionElement(UTypedElementRegistry::GetInstance()->GetElement<UTypedElementWorldInterface>(InElementHandle), GetAssetEditorCustomizationByTypeId(InElementHandle.GetId().GetTypeId()))
+		? FTypedElementViewportInteractionElement(UTypedElementRegistry::GetInstance()->GetElement<UTypedElementWorldInterface>(InElementHandle), GetInterfaceCustomizationByTypeId(InElementHandle.GetId().GetTypeId()))
 		: FTypedElementViewportInteractionElement();
 }
diff --git a/Engine/Source/Editor/EditorFramework/Public/Elements/Framework/TypedElementAssetEditorCustomization.h b/Engine/Source/Editor/EditorFramework/Public/Elements/Framework/TypedElementAssetEditorCustomization.h
deleted file mode 100644
diff --git a/Engine/Source/Editor/EditorFramework/Public/Elements/Framework/TypedElementListObjectUtil.h b/Engine/Source/Editor/EditorFramework/Public/Elements/Framework/TypedElementListObjectUtil.h
deleted file mode 100644
diff --git a/Engine/Source/Editor/EditorFramework/Public/Elements/Framework/TypedElementSelectionSet.h b/Engine/Source/Editor/EditorFramework/Public/Elements/Framework/TypedElementSelectionSet.h
deleted file mode 100644
diff --git a/Engine/Source/Editor/EditorFramework/Public/Elements/Framework/TypedElementViewportInteraction.h b/Engine/Source/Editor/EditorFramework/Public/Elements/Framework/TypedElementViewportInteraction.h
--- a/Engine/Source/Editor/EditorFramework/Public/Elements/Framework/TypedElementViewportInteraction.h
+++ b/Engine/Source/Editor/EditorFramework/Public/Elements/Framework/TypedElementViewportInteraction.h
@@ -4,7 +4,7 @@
 
 #include "InputState.h"
 #include "UnrealWidgetFwd.h"
-#include "Elements/Framework/TypedElementAssetEditorCustomization.h"
+#include "Elements/Framework/TypedElementInterfaceCustomization.h"
 #include "Elements/Interfaces/TypedElementWorldInterface.h"
 #include "TypedElementViewportInteraction.generated.h"
 
@@ -20,10 +20,10 @@ enum class ETypedElementViewportInteractionWorldType : uint8
 /**
  * Customization used to allow asset editors (such as the level editor) to override the base behavior of viewport interaction.
  */
-class EDITORFRAMEWORK_API FTypedElementAssetEditorViewportInteractionCustomization
+class EDITORFRAMEWORK_API FTypedElementViewportInteractionCustomization
 {
 public:
-	virtual ~FTypedElementAssetEditorViewportInteractionCustomization() = default;
+	virtual ~FTypedElementViewportInteractionCustomization() = default;
 
 	//~ See UTypedElementViewportInteraction for API docs
 	virtual void GetElementsToMove(const TTypedElement<UTypedElementWorldInterface>& InElementWorldHandle, const ETypedElementViewportInteractionWorldType InWorldType, const UTypedElementSelectionSet* InSelectionSet, UTypedElementList* OutElementsToMove);
@@ -37,16 +37,16 @@ public:
 };
 
 /**
- * Utility to hold a typed element handle and its associated world interface and asset editor viewport interaction customization.
+ * Utility to hold a typed element handle and its associated world interface and  viewport interaction customization.
  */
 struct EDITORFRAMEWORK_API FTypedElementViewportInteractionElement
 {
 public:
 	FTypedElementViewportInteractionElement() = default;
 
-	FTypedElementViewportInteractionElement(TTypedElement<UTypedElementWorldInterface> InElementWorldHandle, FTypedElementAssetEditorViewportInteractionCustomization* InAssetEditorViewportInteractionCustomization)
+	FTypedElementViewportInteractionElement(TTypedElement<UTypedElementWorldInterface> InElementWorldHandle, FTypedElementViewportInteractionCustomization* InViewportInteractionCustomization)
 		: ElementWorldHandle(MoveTemp(InElementWorldHandle))
-		, AssetEditorViewportInteractionCustomization(InAssetEditorViewportInteractionCustomization)
+		, ViewportInteractionCustomization(InViewportInteractionCustomization)
 	{
 	}
 
@@ -64,28 +64,28 @@ public:
 	FORCEINLINE bool IsSet() const
 	{
 		return ElementWorldHandle.IsSet()
-			&& AssetEditorViewportInteractionCustomization;
+			&& ViewportInteractionCustomization;
 	}
 
 	//~ See UTypedElementViewportInteraction for API docs
-	void GetElementsToMove(const ETypedElementViewportInteractionWorldType InWorldType, const UTypedElementSelectionSet* InSelectionSet, UTypedElementList* OutElementsToMove) { AssetEditorViewportInteractionCustomization->GetElementsToMove(ElementWorldHandle, InWorldType, InSelectionSet, OutElementsToMove); }
-	bool GetGizmoPivotLocation(const UE::Widget::EWidgetMode InWidgetMode, FVector& OutPivotLocation) const { return AssetEditorViewportInteractionCustomization->GetGizmoPivotLocation(ElementWorldHandle, InWidgetMode, OutPivotLocation); }
-	void GizmoManipulationStarted(const UE::Widget::EWidgetMode InWidgetMode) const { AssetEditorViewportInteractionCustomization->GizmoManipulationStarted(ElementWorldHandle, InWidgetMode); }
-	void GizmoManipulationDeltaUpdate(const UE::Widget::EWidgetMode InWidgetMode, const EAxisList::Type InDragAxis, const FInputDeviceState& InInputState, const FTransform& InDeltaTransform, const FVector& InPivotLocation) const { AssetEditorViewportInteractionCustomization->GizmoManipulationDeltaUpdate(ElementWorldHandle, InWidgetMode, InDragAxis, InInputState, InDeltaTransform, InPivotLocation); }
-	void GizmoManipulationStopped(const UE::Widget::EWidgetMode InWidgetMode) const { AssetEditorViewportInteractionCustomization->GizmoManipulationStopped(ElementWorldHandle, InWidgetMode); }
-	void MirrorElement(const FVector& InMirrorScale, const FVector& InPivotLocation) const { AssetEditorViewportInteractionCustomization->MirrorElement(ElementWorldHandle, InMirrorScale, InPivotLocation); }
+	void GetElementsToMove(const ETypedElementViewportInteractionWorldType InWorldType, const UTypedElementSelectionSet* InSelectionSet, UTypedElementList* OutElementsToMove) { ViewportInteractionCustomization->GetElementsToMove(ElementWorldHandle, InWorldType, InSelectionSet, OutElementsToMove); }
+	bool GetGizmoPivotLocation(const UE::Widget::EWidgetMode InWidgetMode, FVector& OutPivotLocation) const { return ViewportInteractionCustomization->GetGizmoPivotLocation(ElementWorldHandle, InWidgetMode, OutPivotLocation); }
+	void GizmoManipulationStarted(const UE::Widget::EWidgetMode InWidgetMode) const { ViewportInteractionCustomization->GizmoManipulationStarted(ElementWorldHandle, InWidgetMode); }
+	void GizmoManipulationDeltaUpdate(const UE::Widget::EWidgetMode InWidgetMode, const EAxisList::Type InDragAxis, const FInputDeviceState& InInputState, const FTransform& InDeltaTransform, const FVector& InPivotLocation) const { ViewportInteractionCustomization->GizmoManipulationDeltaUpdate(ElementWorldHandle, InWidgetMode, InDragAxis, InInputState, InDeltaTransform, InPivotLocation); }
+	void GizmoManipulationStopped(const UE::Widget::EWidgetMode InWidgetMode) const { ViewportInteractionCustomization->GizmoManipulationStopped(ElementWorldHandle, InWidgetMode); }
+	void MirrorElement(const FVector& InMirrorScale, const FVector& InPivotLocation) const { ViewportInteractionCustomization->MirrorElement(ElementWorldHandle, InMirrorScale, InPivotLocation); }
 
 private:
 	TTypedElement<UTypedElementWorldInterface> ElementWorldHandle;
-	FTypedElementAssetEditorViewportInteractionCustomization* AssetEditorViewportInteractionCustomization = nullptr;
+	FTypedElementViewportInteractionCustomization* ViewportInteractionCustomization = nullptr;
 };
 
 /**
  * A utility to handle higher-level viewport interactions, but default via UTypedElementWorldInterface,
- * but asset editors can customize this behavior via FTypedElementAssetEditorViewportInteractionCustomization.
+ * but asset editors can customize this behavior via FTypedElementViewportInteractionCustomization.
  */
 UCLASS(Transient)
-class EDITORFRAMEWORK_API UTypedElementViewportInteraction : public UObject, public TTypedElementAssetEditorCustomizationRegistry<FTypedElementAssetEditorViewportInteractionCustomization>
+class EDITORFRAMEWORK_API UTypedElementViewportInteraction : public UObject, public TTypedElementInterfaceCustomizationRegistry<FTypedElementViewportInteractionCustomization>
 {
 	GENERATED_BODY()
 
@@ -123,7 +123,7 @@ public:
 
 private:
 	/**
-	 * Attempt to resolve the selection interface and asset editor viewport interaction customization for the given element, if any.
+	 * Attempt to resolve the selection interface and viewport interaction customization for the given element, if any.
 	 */
 	FTypedElementViewportInteractionElement ResolveViewportInteractionElement(const FTypedElementHandle& InElementHandle) const;
 };
diff --git a/Engine/Source/Editor/LevelEditor/LevelEditor.Build.cs b/Engine/Source/Editor/LevelEditor/LevelEditor.Build.cs
--- a/Engine/Source/Editor/LevelEditor/LevelEditor.Build.cs
+++ b/Engine/Source/Editor/LevelEditor/LevelEditor.Build.cs
@@ -62,7 +62,7 @@ public class LevelEditor : ModuleRules
 				"Projects",
 				"TargetPlatform",
 				"TypedElementFramework",
-				"TypedElementInterfaces",
+				"TypedElementRuntime",
 				"EngineSettings",
 				"PropertyEditor",
 				"Kismet",
diff --git a/Engine/Source/Editor/LevelEditor/Private/Elements/Actor/ActorElementLevelEditorSelectionCustomization.h b/Engine/Source/Editor/LevelEditor/Private/Elements/Actor/ActorElementLevelEditorSelectionCustomization.h
--- a/Engine/Source/Editor/LevelEditor/Private/Elements/Actor/ActorElementLevelEditorSelectionCustomization.h
+++ b/Engine/Source/Editor/LevelEditor/Private/Elements/Actor/ActorElementLevelEditorSelectionCustomization.h
@@ -7,7 +7,7 @@
 
 class AGroupActor;
 
-class FActorElementLevelEditorSelectionCustomization : public FTypedElementAssetEditorSelectionCustomization, public FTypedElementAssetEditorToolkitHostMixin
+class FActorElementLevelEditorSelectionCustomization : public FTypedElementSelectionCustomization, public FTypedElementAssetEditorToolkitHostMixin
 {
 public:
 	virtual bool CanSelectElement(const TTypedElement<UTypedElementSelectionInterface>& InElementSelectionHandle, const FTypedElementSelectionOptions& InSelectionOptions) override;
diff --git a/Engine/Source/Editor/LevelEditor/Private/Elements/Component/ComponentElementLevelEditorSelectionCustomization.h b/Engine/Source/Editor/LevelEditor/Private/Elements/Component/ComponentElementLevelEditorSelectionCustomization.h
--- a/Engine/Source/Editor/LevelEditor/Private/Elements/Component/ComponentElementLevelEditorSelectionCustomization.h
+++ b/Engine/Source/Editor/LevelEditor/Private/Elements/Component/ComponentElementLevelEditorSelectionCustomization.h
@@ -5,7 +5,7 @@
 #include "Elements/Framework/TypedElementSelectionSet.h"
 #include "Elements/Framework/TypedElementAssetEditorToolkitHostMixin.h"
 
-class FComponentElementLevelEditorSelectionCustomization : public FTypedElementAssetEditorSelectionCustomization, public FTypedElementAssetEditorToolkitHostMixin
+class FComponentElementLevelEditorSelectionCustomization : public FTypedElementSelectionCustomization, public FTypedElementAssetEditorToolkitHostMixin
 {
 public:
 	virtual bool CanSelectElement(const TTypedElement<UTypedElementSelectionInterface>& InElementSelectionHandle, const FTypedElementSelectionOptions& InSelectionOptions) override;
diff --git a/Engine/Source/Editor/LevelEditor/Private/SLevelEditor.cpp b/Engine/Source/Editor/LevelEditor/Private/SLevelEditor.cpp
--- a/Engine/Source/Editor/LevelEditor/Private/SLevelEditor.cpp
+++ b/Engine/Source/Editor/LevelEditor/Private/SLevelEditor.cpp
@@ -229,12 +229,12 @@ void SLevelEditor::Initialize( const TSharedRef<SDockTab>& OwnerTab, const TShar
 	{
 		TUniquePtr<FActorElementLevelEditorSelectionCustomization> ActorCustomization = MakeUnique<FActorElementLevelEditorSelectionCustomization>();
 		ActorCustomization->SetToolkitHost(this);
-		SelectedElements->RegisterAssetEditorCustomizationByTypeName(NAME_Actor, MoveTemp(ActorCustomization));
+		SelectedElements->RegisterInterfaceCustomizationByTypeName(NAME_Actor, MoveTemp(ActorCustomization));
 	}
 	{
 		TUniquePtr<FComponentElementLevelEditorSelectionCustomization> ComponentCustomization = MakeUnique<FComponentElementLevelEditorSelectionCustomization>();
 		ComponentCustomization->SetToolkitHost(this);
-		SelectedElements->RegisterAssetEditorCustomizationByTypeName(NAME_Components, MoveTemp(ComponentCustomization));
+		SelectedElements->RegisterInterfaceCustomizationByTypeName(NAME_Components, MoveTemp(ComponentCustomization));
 	}
 
 	// Allow USelection to bridge to our selected element list
diff --git a/Engine/Source/Editor/UnrealEd/Private/EditorEngine.cpp b/Engine/Source/Editor/UnrealEd/Private/EditorEngine.cpp
--- a/Engine/Source/Editor/UnrealEd/Private/EditorEngine.cpp
+++ b/Engine/Source/Editor/UnrealEd/Private/EditorEngine.cpp
@@ -271,7 +271,7 @@ void InitSelectionSets()
 	GObjectSelection->AddToRoot();
 
 	UTypedElementSelectionSet* ObjectSelectionSet = NewObject<UTypedElementSelectionSet>(GObjectSelection, NAME_None, RF_Transactional);
-	ObjectSelectionSet->RegisterAssetEditorCustomizationByTypeName(NAME_Object, MakeUnique<FObjectElementEditorSelectionCustomization>());
+	ObjectSelectionSet->RegisterInterfaceCustomizationByTypeName(NAME_Object, MakeUnique<FObjectElementEditorSelectionCustomization>());
 	GObjectSelection->SetElementSelectionSet(ObjectSelectionSet);
 }
 
diff --git a/Engine/Source/Editor/UnrealEd/Private/Elements/Actor/ActorElementEditorViewportInteractionCustomization.cpp b/Engine/Source/Editor/UnrealEd/Private/Elements/Actor/ActorElementEditorViewportInteractionCustomization.cpp
--- a/Engine/Source/Editor/UnrealEd/Private/Elements/Actor/ActorElementEditorViewportInteractionCustomization.cpp
+++ b/Engine/Source/Editor/UnrealEd/Private/Elements/Actor/ActorElementEditorViewportInteractionCustomization.cpp
@@ -18,7 +18,7 @@ bool FActorElementEditorViewportInteractionCustomization::GetGizmoPivotLocation(
 		return true;
 	}
 	
-	return FTypedElementAssetEditorViewportInteractionCustomization::GetGizmoPivotLocation(InElementWorldHandle, InWidgetMode, OutPivotLocation);
+	return FTypedElementViewportInteractionCustomization::GetGizmoPivotLocation(InElementWorldHandle, InWidgetMode, OutPivotLocation);
 }
 
 void FActorElementEditorViewportInteractionCustomization::GizmoManipulationDeltaUpdate(const TTypedElement<UTypedElementWorldInterface>& InElementWorldHandle, const UE::Widget::EWidgetMode InWidgetMode, const EAxisList::Type InDragAxis, const FInputDeviceState& InInputState, const FTransform& InDeltaTransform, const FVector& InPivotLocation)
diff --git a/Engine/Source/Editor/UnrealEd/Private/Elements/Component/ComponentElementEditorViewportInteractionCustomization.cpp b/Engine/Source/Editor/UnrealEd/Private/Elements/Component/ComponentElementEditorViewportInteractionCustomization.cpp
--- a/Engine/Source/Editor/UnrealEd/Private/Elements/Component/ComponentElementEditorViewportInteractionCustomization.cpp
+++ b/Engine/Source/Editor/UnrealEd/Private/Elements/Component/ComponentElementEditorViewportInteractionCustomization.cpp
@@ -22,7 +22,7 @@ bool FComponentElementEditorViewportInteractionCustomization::GetGizmoPivotLocat
 		}
 		else
 		{
-			FTypedElementAssetEditorViewportInteractionCustomization::GetGizmoPivotLocation(InElementWorldHandle, InWidgetMode, BasePivotLocation);
+			FTypedElementViewportInteractionCustomization::GetGizmoPivotLocation(InElementWorldHandle, InWidgetMode, BasePivotLocation);
 		}
 
 		// If necessary, transform the editor pivot location to be relative to the component's parent
diff --git a/Engine/Source/Editor/UnrealEd/Private/Elements/Object/ObjectElementEditorSelectionCustomization.h b/Engine/Source/Editor/UnrealEd/Private/Elements/Object/ObjectElementEditorSelectionCustomization.h
--- a/Engine/Source/Editor/UnrealEd/Private/Elements/Object/ObjectElementEditorSelectionCustomization.h
+++ b/Engine/Source/Editor/UnrealEd/Private/Elements/Object/ObjectElementEditorSelectionCustomization.h
@@ -4,7 +4,7 @@
 
 #include "Elements/Framework/TypedElementSelectionSet.h"
 
-class FObjectElementEditorSelectionCustomization : public FTypedElementAssetEditorSelectionCustomization
+class FObjectElementEditorSelectionCustomization : public FTypedElementSelectionCustomization
 {
 public:
 	virtual bool SelectElement(const TTypedElement<UTypedElementSelectionInterface>& InElementSelectionHandle, UTypedElementList* InSelectionSet, const FTypedElementSelectionOptions& InSelectionOptions) override;
diff --git a/Engine/Source/Editor/UnrealEd/Private/LevelEditorViewport.cpp b/Engine/Source/Editor/UnrealEd/Private/LevelEditorViewport.cpp
--- a/Engine/Source/Editor/UnrealEd/Private/LevelEditorViewport.cpp
+++ b/Engine/Source/Editor/UnrealEd/Private/LevelEditorViewport.cpp
@@ -1816,13 +1816,13 @@ void FLevelEditorViewportClient::InitializeViewportInteraction()
 		TUniquePtr<FActorElementLevelEditorViewportInteractionCustomization> ActorCustomization = MakeUnique<FActorElementLevelEditorViewportInteractionCustomization>();
 		ActorCustomization->SetLevelEditorViewportClient(this);
 		ActorCustomization->SetToolkitHost(ParentLevelEditor.Pin().Get());
-		ViewportInteraction->RegisterAssetEditorCustomizationByTypeName(NAME_Actor, MoveTemp(ActorCustomization));
+		ViewportInteraction->RegisterInterfaceCustomizationByTypeName(NAME_Actor, MoveTemp(ActorCustomization));
 	}
 	{
 		TUniquePtr<FComponentElementLevelEditorViewportInteractionCustomization> ComponentCustomization = MakeUnique<FComponentElementLevelEditorViewportInteractionCustomization>();
 		ComponentCustomization->SetLevelEditorViewportClient(this);
 		ComponentCustomization->SetToolkitHost(ParentLevelEditor.Pin().Get());
-		ViewportInteraction->RegisterAssetEditorCustomizationByTypeName(NAME_Components, MoveTemp(ComponentCustomization));
+		ViewportInteraction->RegisterInterfaceCustomizationByTypeName(NAME_Components, MoveTemp(ComponentCustomization));
 	}
 }
 
diff --git a/Engine/Source/Editor/UnrealEd/Public/Elements/Actor/ActorElementEditorViewportInteractionCustomization.h b/Engine/Source/Editor/UnrealEd/Public/Elements/Actor/ActorElementEditorViewportInteractionCustomization.h
--- a/Engine/Source/Editor/UnrealEd/Public/Elements/Actor/ActorElementEditorViewportInteractionCustomization.h
+++ b/Engine/Source/Editor/UnrealEd/Public/Elements/Actor/ActorElementEditorViewportInteractionCustomization.h
@@ -7,7 +7,7 @@
 
 class AActor;
 
-class UNREALED_API FActorElementEditorViewportInteractionCustomization : public FTypedElementAssetEditorViewportInteractionCustomization, public FTypedElementAssetEditorToolkitHostMixin
+class UNREALED_API FActorElementEditorViewportInteractionCustomization : public FTypedElementViewportInteractionCustomization, public FTypedElementAssetEditorToolkitHostMixin
 {
 public:
 	virtual bool GetGizmoPivotLocation(const TTypedElement<UTypedElementWorldInterface>& InElementWorldHandle, const UE::Widget::EWidgetMode InWidgetMode, FVector& OutPivotLocation) override;
diff --git a/Engine/Source/Editor/UnrealEd/Public/Elements/Component/ComponentElementEditorViewportInteractionCustomization.h b/Engine/Source/Editor/UnrealEd/Public/Elements/Component/ComponentElementEditorViewportInteractionCustomization.h
--- a/Engine/Source/Editor/UnrealEd/Public/Elements/Component/ComponentElementEditorViewportInteractionCustomization.h
+++ b/Engine/Source/Editor/UnrealEd/Public/Elements/Component/ComponentElementEditorViewportInteractionCustomization.h
@@ -8,7 +8,7 @@
 class FProperty;
 class USceneComponent;
 
-class UNREALED_API FComponentElementEditorViewportInteractionCustomization : public FTypedElementAssetEditorViewportInteractionCustomization, public FTypedElementAssetEditorToolkitHostMixin
+class UNREALED_API FComponentElementEditorViewportInteractionCustomization : public FTypedElementViewportInteractionCustomization, public FTypedElementAssetEditorToolkitHostMixin
 {
 public:
 	virtual bool GetGizmoPivotLocation(const TTypedElement<UTypedElementWorldInterface>& InElementWorldHandle, const UE::Widget::EWidgetMode InWidgetMode, FVector& OutPivotLocation) override;
diff --git a/Engine/Source/Editor/UnrealEd/UnrealEd.Build.cs b/Engine/Source/Editor/UnrealEd/UnrealEd.Build.cs
--- a/Engine/Source/Editor/UnrealEd/UnrealEd.Build.cs
+++ b/Engine/Source/Editor/UnrealEd/UnrealEd.Build.cs
@@ -109,7 +109,7 @@ public class UnrealEd : ModuleRules
 				"EditorSubsystem",
 				"InteractiveToolsFramework",
 				"TypedElementFramework",
-				"TypedElementInterfaces",
+				"TypedElementRuntime",
 				"ToolMenusEditor",
 				"StatusBar",
 				"InterchangeCore",
diff --git a/Engine/Source/Runtime/Engine/Engine.Build.cs b/Engine/Source/Runtime/Engine/Engine.Build.cs
--- a/Engine/Source/Runtime/Engine/Engine.Build.cs
+++ b/Engine/Source/Runtime/Engine/Engine.Build.cs
@@ -101,6 +101,7 @@ public class Engine : ModuleRules
 		PublicIncludePathModuleNames.AddRange(
 			new string[] {
 				"TypedElementFramework",
+				"TypedElementRuntime",
 			}
 		);
 
@@ -113,7 +114,7 @@ public class Engine : ModuleRules
 				"UMG",
 				"Projects",
 				"TypedElementFramework",
-				"TypedElementInterfaces",
+				"TypedElementRuntime",
 				"MaterialShaderQualitySettings",
 				"CinematicCamera",
 				"Analytics",
diff --git a/Engine/Source/Runtime/Engine/Private/Elements/Framework/EngineElements.cpp b/Engine/Source/Runtime/Engine/Private/Elements/Framework/EngineElements.cpp
--- a/Engine/Source/Runtime/Engine/Private/Elements/Framework/EngineElements.cpp
+++ b/Engine/Source/Runtime/Engine/Private/Elements/Framework/EngineElements.cpp
@@ -58,7 +58,7 @@ void RegisterEngineElements()
 {
 	// Ensure the framework and base interfaces are also loaded
 	FModuleManager::Get().LoadModuleChecked("TypedElementFramework");
-	FModuleManager::Get().LoadModuleChecked("TypedElementInterfaces");
+	FModuleManager::Get().LoadModuleChecked("TypedElementRuntime");
 
 	RegisterEngineObjectElements();
 	RegisterEngineActorElements();
diff --git a/Engine/Source/Runtime/TypedElementInterfaces/Private/Elements/Interfaces/TypedElementAssetDataInterface.cpp b/Engine/Source/Runtime/TypedElementInterfaces/Private/Elements/Interfaces/TypedElementAssetDataInterface.cpp
deleted file mode 100644
diff --git a/Engine/Source/Runtime/TypedElementInterfaces/Private/Elements/Interfaces/TypedElementObjectInterface.cpp b/Engine/Source/Runtime/TypedElementInterfaces/Private/Elements/Interfaces/TypedElementObjectInterface.cpp
deleted file mode 100644
diff --git a/Engine/Source/Runtime/TypedElementInterfaces/Private/Elements/Interfaces/TypedElementSelectionInterface.cpp b/Engine/Source/Runtime/TypedElementInterfaces/Private/Elements/Interfaces/TypedElementSelectionInterface.cpp
deleted file mode 100644
diff --git a/Engine/Source/Runtime/TypedElementInterfaces/Private/TypedElementInterfacesModule.cpp b/Engine/Source/Runtime/TypedElementInterfaces/Private/TypedElementInterfacesModule.cpp
deleted file mode 100644
diff --git a/Engine/Source/Runtime/TypedElementInterfaces/Public/Elements/Interfaces/TypedElementAssetDataInterface.h b/Engine/Source/Runtime/TypedElementInterfaces/Public/Elements/Interfaces/TypedElementAssetDataInterface.h
deleted file mode 100644
diff --git a/Engine/Source/Runtime/TypedElementInterfaces/Public/Elements/Interfaces/TypedElementObjectInterface.h b/Engine/Source/Runtime/TypedElementInterfaces/Public/Elements/Interfaces/TypedElementObjectInterface.h
deleted file mode 100644
diff --git a/Engine/Source/Runtime/TypedElementInterfaces/Public/Elements/Interfaces/TypedElementSelectionInterface.h b/Engine/Source/Runtime/TypedElementInterfaces/Public/Elements/Interfaces/TypedElementSelectionInterface.h
deleted file mode 100644
diff --git a/Engine/Source/Runtime/TypedElementInterfaces/TypedElementInterfaces.Build.cs b/Engine/Source/Runtime/TypedElementInterfaces/TypedElementInterfaces.Build.cs
deleted file mode 100644
diff --git a/Engine/Source/Runtime/TypedElementRuntime/Private/Elements/Framework/TypedElementInterfaceCustomization.cpp b/Engine/Source/Runtime/TypedElementRuntime/Private/Elements/Framework/TypedElementInterfaceCustomization.cpp
new file mode 100644
--- /dev/null
+++ b/Engine/Source/Runtime/TypedElementRuntime/Private/Elements/Framework/TypedElementInterfaceCustomization.cpp
@@ -0,0 +1,16 @@
+// Copyright Epic Games, Inc. All Rights Reserved.
+
+#include "Elements/Framework/TypedElementInterfaceCustomization.h"
+#include "Elements/Framework/TypedElementRegistry.h"
+
+FTypedHandleTypeId FTypedElementInterfaceCustomizationRegistryBase::GetElementTypeIdFromName(const FName InElementTypeName) const
+{
+	return UTypedElementRegistry::GetInstance()->GetRegisteredElementTypeId(InElementTypeName);
+}
+
+FTypedHandleTypeId FTypedElementInterfaceCustomizationRegistryBase::GetElementTypeIdFromNameChecked(const FName InElementTypeName) const
+{
+	const FTypedHandleTypeId ElementTypeId = GetElementTypeIdFromName(InElementTypeName);
+	checkf(ElementTypeId > 0, TEXT("Element type '%s' has not been registered!"), *InElementTypeName.ToString());
+	return ElementTypeId;
+}
diff --git a/Engine/Source/Runtime/TypedElementRuntime/Private/Elements/Framework/TypedElementListObjectUtil.cpp b/Engine/Source/Runtime/TypedElementRuntime/Private/Elements/Framework/TypedElementListObjectUtil.cpp
new file mode 100644
--- /dev/null
+++ b/Engine/Source/Runtime/TypedElementRuntime/Private/Elements/Framework/TypedElementListObjectUtil.cpp
@@ -0,0 +1,101 @@
+// Copyright Epic Games, Inc. All Rights Reserved.
+
+#include "Elements/Framework/TypedElementListObjectUtil.h"
+#include "Elements/Interfaces/TypedElementObjectInterface.h"
+
+namespace TypedElementListObjectUtil
+{
+
+UObject* GetObjectOfType(const TTypedElement<UTypedElementObjectInterface>& InObjectElement, const UClass* InRequiredClass)
+{
+	UObject* ElementObject = InObjectElement ? InObjectElement.GetObject() : nullptr;
+	return (ElementObject && (!InRequiredClass || ElementObject->IsA(InRequiredClass)))
+		? ElementObject
+		: nullptr;
+}
+
+bool HasObjects(const UTypedElementList* InElementList, const UClass* InRequiredClass)
+{
+	bool bHasObjects = false;
+
+	ForEachObject(InElementList, [&bHasObjects](const UObject*)
+	{
+		bHasObjects = true;
+		return false;
+	}, InRequiredClass);
+
+	return bHasObjects;
+}
+
+int32 CountObjects(const UTypedElementList* InElementList, const UClass* InRequiredClass)
+{
+	int32 NumObjects = 0;
+
+	ForEachObject(InElementList, [&NumObjects](const UObject*)
+	{
+		++NumObjects;
+		return true;
+	}, InRequiredClass);
+
+	return NumObjects;
+}
+
+void ForEachObject(const UTypedElementList* InElementList, TFunctionRef<bool(UObject*)> InCallback, const UClass* InRequiredClass)
+{
+	InElementList->ForEachElement<UTypedElementObjectInterface>([&InCallback, InRequiredClass](const TTypedElement<UTypedElementObjectInterface>& InObjectElement)
+	{
+		if (UObject* ElementObject = GetObjectOfType(InObjectElement, InRequiredClass))
+		{
+			return InCallback(ElementObject);
+		}
+		return true;
+	});
+}
+
+TArray<UObject*> GetObjects(const UTypedElementList* InElementList, const UClass* InRequiredClass)
+{
+	TArray<UObject*> SelectedObjects;
+	SelectedObjects.Reserve(InElementList->Num());
+
+	ForEachObject(InElementList, [&SelectedObjects](UObject* InObject)
+	{
+		SelectedObjects.Add(InObject);
+		return true;
+	}, InRequiredClass);
+
+	return SelectedObjects;
+}
+
+UObject* GetTopObject(const UTypedElementList* InElementList, const UClass* InRequiredClass)
+{
+	TTypedElement<UTypedElementObjectInterface> TempElement;
+	for (int32 ElementIndex = 0; ElementIndex < InElementList->Num(); ++ElementIndex)
+	{
+		InElementList->GetElementAt(ElementIndex, TempElement);
+		
+		if (UObject* ElementObject = GetObjectOfType(TempElement, InRequiredClass))
+		{
+			return ElementObject;
+		}
+	}
+
+	return nullptr;
+}
+
+UObject* GetBottomObject(const UTypedElementList* InElementList, const UClass* InRequiredClass)
+{
+	TTypedElement<UTypedElementObjectInterface> TempElement;
+	for (int32 ElementIndex = InElementList->Num() - 1; ElementIndex >= 0; --ElementIndex)
+	{
+		InElementList->GetElementAt(ElementIndex, TempElement);
+
+		if (UObject* ElementObject = GetObjectOfType(TempElement, InRequiredClass))
+		{
+			return ElementObject;
+		}
+	}
+
+	return nullptr;
+}
+
+} // namespace TypedElementListObjectUtil
diff --git a/Engine/Source/Runtime/TypedElementRuntime/Private/Elements/Framework/TypedElementSelectionSet.cpp b/Engine/Source/Runtime/TypedElementRuntime/Private/Elements/Framework/TypedElementSelectionSet.cpp
new file mode 100644
--- /dev/null
+++ b/Engine/Source/Runtime/TypedElementRuntime/Private/Elements/Framework/TypedElementSelectionSet.cpp
@@ -0,0 +1,255 @@
+// Copyright Epic Games, Inc. All Rights Reserved.
+
+#include "Elements/Framework/TypedElementSelectionSet.h"
+#include "Elements/Framework/TypedElementRegistry.h"
+
+UTypedElementSelectionSet::UTypedElementSelectionSet()
+{
+	if (!HasAnyFlags(RF_ClassDefaultObject))
+	{
+		ElementList = UTypedElementRegistry::GetInstance()->CreateElementList();
+		ElementList->OnPreChange().AddUObject(this, &UTypedElementSelectionSet::OnElementListPreChange);
+		ElementList->OnChanged().AddUObject(this, &UTypedElementSelectionSet::OnElementListChanged);
+	}
+}
+
+#if WITH_EDITOR
+bool UTypedElementSelectionSet::Modify(bool bAlwaysMarkDirty)
+{
+	if (GUndo && CanModify())
+	{
+		bool bCanModify = true;
+		ElementList->ForEachElement<UTypedElementSelectionInterface>([&bCanModify](const TTypedElement<UTypedElementSelectionInterface>& InSelectionElement)
+		{
+			bCanModify = !InSelectionElement.ShouldPreventTransactions();
+			return bCanModify;
+		});
+
+		if (!bCanModify)
+		{
+			return false;
+		}
+
+		return Super::Modify(bAlwaysMarkDirty);
+	}
+
+	return false;
+}
+#endif	// WITH_EDITOR
+
+void UTypedElementSelectionSet::Serialize(FArchive& Ar)
+{
+	checkf(!Ar.IsPersistent(), TEXT("UTypedElementSelectionSet can only be serialized by transient archives!"));
+
+	if (Ar.IsSaving())
+	{
+		FTypedHandleTypeId ElementTypeId = 0;
+
+		ElementList->ForEachElement<UTypedElementSelectionInterface>([&Ar, &ElementTypeId](const TTypedElement<UTypedElementSelectionInterface>& InSelectionElement)
+		{
+			ElementTypeId = InSelectionElement.GetId().GetTypeId();
+			Ar << ElementTypeId;
+
+			InSelectionElement.WriteTransactedElement(Ar);
+			return true;
+		});
+
+		// End of the list
+		ElementTypeId = 0;
+		Ar << ElementTypeId;
+	}
+	else if (Ar.IsLoading())
+	{
+		TArray<FTypedElementHandle, TInlineAllocator<256>> SelectedElements;
+
+		{
+			UTypedElementRegistry* Registry = UTypedElementRegistry::GetInstance();
+
+			FTypedHandleTypeId ElementTypeId = 0;
+			for (;;)
+			{
+				Ar << ElementTypeId;
+				if (!ElementTypeId)
+				{
+					// End of the list
+					break;
+				}
+
+				UTypedElementSelectionInterface* ElementTypeSelectionInterface = Registry->GetElementInterface<UTypedElementSelectionInterface>(ElementTypeId);
+				checkf(ElementTypeSelectionInterface, TEXT("Failed to find selection interface for a previously transacted element type!"));
+				FTypedElementHandle SelectedElement = ElementTypeSelectionInterface->ReadTransactedElement(Ar);
+				if (SelectedElement)
+				{
+					SelectedElements.Add(MoveTemp(SelectedElement));
+				}
+			}
+		}
+
+		{
+			FTypedElementListLegacySyncScopedBatch LegacySyncBatch(ElementList, /*bNotify*/false);
+			TGuardValue<bool> GuardIsRestoringFromTransaction(bIsRestoringFromTransaction, true);
+
+			// TODO: Work out the intersection of the before and after state instead of clearing and reselecting?
+
+			const FTypedElementSelectionOptions SelectionOptions = FTypedElementSelectionOptions()
+				.SetAllowHidden(true)
+				.SetAllowGroups(false)
+				.SetWarnIfLocked(false);
+
+			ClearSelection(SelectionOptions);
+			SelectElements(SelectedElements, SelectionOptions);
+		}
+	}
+}
+
+bool UTypedElementSelectionSet::IsElementSelected(const FTypedElementHandle& InElementHandle, const FTypedElementIsSelectedOptions InSelectionOptions) const
+{
+	FTypedElementSelectionSetElement SelectionSetElement = ResolveSelectionSetElement(InElementHandle);
+	return SelectionSetElement && SelectionSetElement.IsElementSelected(InSelectionOptions);
+}
+
+bool UTypedElementSelectionSet::CanSelectElement(const FTypedElementHandle& InElementHandle, const FTypedElementSelectionOptions InSelectionOptions) const
+{
+	FTypedElementSelectionSetElement SelectionSetElement = ResolveSelectionSetElement(InElementHandle);
+	return SelectionSetElement && SelectionSetElement.CanSelectElement(InSelectionOptions);
+}
+
+bool UTypedElementSelectionSet::CanDeselectElement(const FTypedElementHandle& InElementHandle, const FTypedElementSelectionOptions InSelectionOptions) const
+{
+	FTypedElementSelectionSetElement SelectionSetElement = ResolveSelectionSetElement(InElementHandle);
+	return SelectionSetElement && SelectionSetElement.CanDeselectElement(InSelectionOptions);
+}
+
+bool UTypedElementSelectionSet::SelectElement(const FTypedElementHandle& InElementHandle, const FTypedElementSelectionOptions InSelectionOptions)
+{
+	FTypedElementSelectionSetElement SelectionSetElement = ResolveSelectionSetElement(InElementHandle);
+	return SelectionSetElement && SelectionSetElement.CanSelectElement(InSelectionOptions) && SelectionSetElement.SelectElement(InSelectionOptions);
+}
+
+bool UTypedElementSelectionSet::SelectElements(const TArray<FTypedElementHandle>& InElementHandles, const FTypedElementSelectionOptions InSelectionOptions)
+{
+	return SelectElements(MakeArrayView(InElementHandles), InSelectionOptions);
+}
+
+bool UTypedElementSelectionSet::SelectElements(TArrayView<const FTypedElementHandle> InElementHandles, const FTypedElementSelectionOptions InSelectionOptions)
+{
+	FTypedElementListLegacySyncScopedBatch LegacySyncBatch(ElementList, InSelectionOptions.AllowLegacyNotifications());
+
+	bool bSelectionChanged = false;
+
+	for (const FTypedElementHandle& ElementHandle : InElementHandles)
+	{
+		bSelectionChanged |= SelectElement(ElementHandle, InSelectionOptions);
+	}
+
+	return bSelectionChanged;
+}
+
+bool UTypedElementSelectionSet::DeselectElement(const FTypedElementHandle& InElementHandle, const FTypedElementSelectionOptions InSelectionOptions)
+{
+	FTypedElementSelectionSetElement SelectionSetElement = ResolveSelectionSetElement(InElementHandle);
+	return SelectionSetElement && SelectionSetElement.CanDeselectElement(InSelectionOptions) && SelectionSetElement.DeselectElement(InSelectionOptions);
+}
+
+bool UTypedElementSelectionSet::DeselectElements(const TArray<FTypedElementHandle>& InElementHandles, const FTypedElementSelectionOptions InSelectionOptions)
+{
+	return DeselectElements(MakeArrayView(InElementHandles), InSelectionOptions);
+}
+
+bool UTypedElementSelectionSet::DeselectElements(TArrayView<const FTypedElementHandle> InElementHandles, const FTypedElementSelectionOptions InSelectionOptions)
+{
+	FTypedElementListLegacySyncScopedBatch LegacySyncBatch(ElementList, InSelectionOptions.AllowLegacyNotifications());
+
+	bool bSelectionChanged = false;
+
+	for (const FTypedElementHandle& ElementHandle : InElementHandles)
+	{
+		bSelectionChanged |= DeselectElement(ElementHandle, InSelectionOptions);
+	}
+
+	return bSelectionChanged;
+}
+
+bool UTypedElementSelectionSet::ClearSelection(const FTypedElementSelectionOptions InSelectionOptions)
+{
+	FTypedElementListLegacySyncScopedBatch LegacySyncBatch(ElementList, InSelectionOptions.AllowLegacyNotifications());
+
+	bool bSelectionChanged = false;
+
+	// Run deselection via the selection interface where possible
+	{
+		// Take a copy of the currently selected elements to avoid mutating the selection set while iterating
+		TArray<FTypedElementHandle, TInlineAllocator<8>> ElementsCopy;
+		ElementList->GetElementHandles(ElementsCopy);
+
+		for (const FTypedElementHandle& ElementHandle : ElementsCopy)
+		{
+			bSelectionChanged |= DeselectElement(ElementHandle, InSelectionOptions);
+		}
+	}
+
+	// TODO: BSP surfaces?
+
+	// If anything remains in the selection set after processing elements that implement this interface, just clear it
+	if (ElementList->Num() > 0)
+	{
+		bSelectionChanged = true;
+		ElementList->Reset();
+	}
+
+	return bSelectionChanged;
+}
+
+bool UTypedElementSelectionSet::SetSelection(const TArray<FTypedElementHandle>& InElementHandles, const FTypedElementSelectionOptions InSelectionOptions)
+{
+	return SetSelection(MakeArrayView(InElementHandles), InSelectionOptions);
+}
+
+bool UTypedElementSelectionSet::SetSelection(TArrayView<const FTypedElementHandle> InElementHandles, const FTypedElementSelectionOptions InSelectionOptions)
+{
+	FTypedElementListLegacySyncScopedBatch LegacySyncBatch(ElementList, InSelectionOptions.AllowLegacyNotifications());
+
+	bool bSelectionChanged = false;
+
+	bSelectionChanged |= ClearSelection(InSelectionOptions);
+	bSelectionChanged |= SelectElements(InElementHandles, InSelectionOptions);
+
+	return bSelectionChanged;
+}
+
+bool UTypedElementSelectionSet::AllowSelectionModifiers(const FTypedElementHandle& InElementHandle) const
+{
+	FTypedElementSelectionSetElement SelectionSetElement = ResolveSelectionSetElement(InElementHandle);
+	return SelectionSetElement && SelectionSetElement.AllowSelectionModifiers();
+}
+
+FTypedElementHandle UTypedElementSelectionSet::GetSelectionElement(const FTypedElementHandle& InElementHandle, const ETypedElementSelectionMethod InSelectionMethod) const
+{
+	FTypedElementSelectionSetElement SelectionSetElement = ResolveSelectionSetElement(InElementHandle);
+	return SelectionSetElement ? SelectionSetElement.GetSelectionElement(InSelectionMethod) : FTypedElementHandle();
+}
+
+FTypedElementSelectionSetElement UTypedElementSelectionSet::ResolveSelectionSetElement(const FTypedElementHandle& InElementHandle) const
+{
+	return InElementHandle
+		? FTypedElementSelectionSetElement(ElementList->GetElement<UTypedElementSelectionInterface>(InElementHandle), ElementList, GetInterfaceCustomizationByTypeId(InElementHandle.GetId().GetTypeId()))
+		: FTypedElementSelectionSetElement();
+}
+
+void UTypedElementSelectionSet::OnElementListPreChange(const UTypedElementList* InElementList)
+{
+	check(InElementList == ElementList);
+	OnPreChangeDelegate.Broadcast(this);
+
+	if (!bIsRestoringFromTransaction)
+	{
+		// Track the pre-change state for undo/redo
+		Modify();
+	}
+}
+
+void UTypedElementSelectionSet::OnElementListChanged(const UTypedElementList* InElementList)
+{
+	check(InElementList == ElementList);
+	OnChangedDelegate.Broadcast(this);
+}
diff --git a/Engine/Source/Runtime/TypedElementRuntime/Private/Elements/Interfaces/TypedElementAssetDataInterface.cpp b/Engine/Source/Runtime/TypedElementRuntime/Private/Elements/Interfaces/TypedElementAssetDataInterface.cpp
new file mode 100644
--- /dev/null
+++ b/Engine/Source/Runtime/TypedElementRuntime/Private/Elements/Interfaces/TypedElementAssetDataInterface.cpp
@@ -0,0 +1,21 @@
+// Copyright Epic Games, Inc. All Rights Reserved.
+
+#include "Elements/Interfaces/TypedElementAssetDataInterface.h"
+
+TArray<FAssetData> UTypedElementAssetDataInterface::GetAllReferencedAssetDatas(const FTypedElementHandle& InElementHandle)
+{
+	TArray<FAssetData> AssetDatas;
+
+	FAssetData ElementAssetData = GetAssetData(InElementHandle);
+	if (ElementAssetData.IsValid())
+	{
+		AssetDatas.Emplace(ElementAssetData);
+	}
+
+	return AssetDatas;
+}
+
+FAssetData UTypedElementAssetDataInterface::GetAssetData(const FTypedElementHandle& InElementHandle)
+{
+	return FAssetData();
+}
diff --git a/Engine/Source/Runtime/TypedElementRuntime/Private/Elements/Interfaces/TypedElementObjectInterface.cpp b/Engine/Source/Runtime/TypedElementRuntime/Private/Elements/Interfaces/TypedElementObjectInterface.cpp
new file mode 100644
--- /dev/null
+++ b/Engine/Source/Runtime/TypedElementRuntime/Private/Elements/Interfaces/TypedElementObjectInterface.cpp
@@ -0,0 +1,16 @@
+// Copyright Epic Games, Inc. All Rights Reserved.
+
+#include "Elements/Interfaces/TypedElementObjectInterface.h"
+
+#include "Elements/Framework/TypedElementRegistry.h"
+
+UObject* UTypedElementObjectInterface::GetObject(const FTypedElementHandle& InElementHandle)
+{
+	return nullptr;
+}
+
+UClass* UTypedElementObjectInterface::GetObjectClass(const FTypedElementHandle& InElementHandle)
+{
+	UObject* HandleAsObject = GetObject(InElementHandle);
+	return HandleAsObject ? HandleAsObject->GetClass() : nullptr;
+}
diff --git a/Engine/Source/Runtime/TypedElementRuntime/Private/Elements/Interfaces/TypedElementSelectionInterface.cpp b/Engine/Source/Runtime/TypedElementRuntime/Private/Elements/Interfaces/TypedElementSelectionInterface.cpp
new file mode 100644
--- /dev/null
+++ b/Engine/Source/Runtime/TypedElementRuntime/Private/Elements/Interfaces/TypedElementSelectionInterface.cpp
@@ -0,0 +1,19 @@
+// Copyright Epic Games, Inc. All Rights Reserved.
+
+#include "Elements/Interfaces/TypedElementSelectionInterface.h"
+#include "Elements/Framework/TypedElementList.h"
+
+bool UTypedElementSelectionInterface::IsElementSelected(const FTypedElementHandle& InElementHandle, const UTypedElementList* InSelectionSet, const FTypedElementIsSelectedOptions& InSelectionOptions)
+{
+	return InSelectionSet->Contains(InElementHandle);
+}
+
+bool UTypedElementSelectionInterface::SelectElement(const FTypedElementHandle& InElementHandle, UTypedElementList* InSelectionSet, const FTypedElementSelectionOptions& InSelectionOptions)
+{
+	return InSelectionSet->Add(InElementHandle);
+}
+
+bool UTypedElementSelectionInterface::DeselectElement(const FTypedElementHandle& InElementHandle, UTypedElementList* InSelectionSet, const FTypedElementSelectionOptions& InSelectionOptions)
+{
+	return InSelectionSet->Remove(InElementHandle);
+}
diff --git a/Engine/Source/Runtime/TypedElementRuntime/Private/TypedElementRuntimeModule.cpp b/Engine/Source/Runtime/TypedElementRuntime/Private/TypedElementRuntimeModule.cpp
new file mode 100644
--- /dev/null
+++ b/Engine/Source/Runtime/TypedElementRuntime/Private/TypedElementRuntimeModule.cpp
@@ -0,0 +1,14 @@
+// Copyright Epic Games, Inc. All Rights Reserved.
+
+#include "Modules/ModuleManager.h"
+
+class FTypedElementRuntimeModule : public FDefaultModuleImpl
+{
+public:
+	virtual void StartupModule() override
+	{
+		FModuleManager::Get().LoadModuleChecked("TypedElementFramework");
+	}
+};
+
+IMPLEMENT_MODULE(FTypedElementRuntimeModule, TypedElementRuntime)
diff --git a/Engine/Source/Runtime/TypedElementRuntime/Public/Elements/Framework/TypedElementInterfaceCustomization.h b/Engine/Source/Runtime/TypedElementRuntime/Public/Elements/Framework/TypedElementInterfaceCustomization.h
new file mode 100644
--- /dev/null
+++ b/Engine/Source/Runtime/TypedElementRuntime/Public/Elements/Framework/TypedElementInterfaceCustomization.h
@@ -0,0 +1,122 @@
+// Copyright Epic Games, Inc. All Rights Reserved.
+
+#pragma once
+
+#include "CoreMinimal.h"
+#include "Templates/UniquePtr.h"
+#include "Elements/Framework/TypedElementLimits.h"
+
+/**
+ * Non-templated base class for the interface customization registry.
+ */
+class TYPEDELEMENTRUNTIME_API FTypedElementInterfaceCustomizationRegistryBase
+{
+public:
+	virtual ~FTypedElementInterfaceCustomizationRegistryBase() = default;
+
+protected:
+	/**
+	 * Given an element name, attempt to get its registered type ID from the global registry.
+	 * @return The registered type ID, or 0 if the element name is not registered.
+	 */
+	FTypedHandleTypeId GetElementTypeIdFromName(const FName InElementTypeName) const;
+
+	/**
+	 * Given an element name, attempt to get its registered type ID from the global registry.
+	 * @return The registered type ID, or asserts if the element name is not registered.
+	 */
+	FTypedHandleTypeId GetElementTypeIdFromNameChecked(const FName InElementTypeName) const;
+};
+
+/**
+ * Utility to register and retrieve interface customizations for a given type.
+ */
+template <typename CustomizationBaseType, typename DefaultCustomizationType = CustomizationBaseType>
+class TTypedElementInterfaceCustomizationRegistry : public FTypedElementInterfaceCustomizationRegistryBase
+{
+public:
+	template <typename... TDefaultArgs>
+	TTypedElementInterfaceCustomizationRegistry(TDefaultArgs&&... DefaultArgs)
+		: DefaultInterfaceCustomization(MakeUnique<DefaultCustomizationType>(Forward<TDefaultArgs>(DefaultArgs)...))
+	{
+	}
+
+	virtual ~TTypedElementInterfaceCustomizationRegistry() = default;
+
+	TTypedElementInterfaceCustomizationRegistry(const TTypedElementInterfaceCustomizationRegistry&) = delete;
+	TTypedElementInterfaceCustomizationRegistry& operator=(const TTypedElementInterfaceCustomizationRegistry&) = delete;
+
+	TTypedElementInterfaceCustomizationRegistry(TTypedElementInterfaceCustomizationRegistry&&) = delete;
+	TTypedElementInterfaceCustomizationRegistry& operator=(TTypedElementInterfaceCustomizationRegistry&&) = delete;
+
+	/**
+	 * Set the default interface customization instance.
+	 */
+	void SetDefaultInterfaceCustomization(TUniquePtr<CustomizationBaseType>&& InInterfaceCustomization)
+	{
+		checkf(InInterfaceCustomization, TEXT("Default interface customization cannot be null!"));
+		DefaultInterfaceCustomization = MoveTemp(InInterfaceCustomization);
+	}
+
+	/**
+	 * Register an interface customization for the given element type.
+	 */
+	void RegisterInterfaceCustomizationByTypeName(const FName InElementTypeName, TUniquePtr<CustomizationBaseType>&& InInterfaceCustomization)
+	{
+		RegisterInterfaceCustomizationByTypeId(GetElementTypeIdFromNameChecked(InElementTypeName), MoveTemp(InInterfaceCustomization));
+	}
+
+	/**
+	 * Register an interface customization for the given element type.
+	 */
+	void RegisterInterfaceCustomizationByTypeId(const FTypedHandleTypeId InElementTypeId, TUniquePtr<CustomizationBaseType>&& InInterfaceCustomization)
+	{
+		RegisteredInterfaceCustomizations[InElementTypeId - 1] = MoveTemp(InInterfaceCustomization);
+	}
+
+	/**
+	 * Unregister an interface customization for the given element type.
+	 */
+	void UnregisterInterfaceCustomizationByTypeName(const FName InElementTypeName)
+	{
+		UnregisterInterfaceCustomizationByTypeId(GetElementTypeIdFromNameChecked(InElementTypeName));
+	}
+
+	/**
+	 * Unregister an interface customization for the given element type.
+	 */
+	void UnregisterInterfaceCustomizationByTypeId(const FTypedHandleTypeId InElementTypeId)
+	{
+		RegisteredInterfaceCustomizations[InElementTypeId - 1].Reset();
+	}
+
+	/**
+	 * Get the interface customization for the given element type.
+	 * @note If bAllowFallback is true, then this will return the default interface customization if no override is present, otherwise it will return null.
+	 */
+	CustomizationBaseType* GetInterfaceCustomizationByTypeName(const FName InElementTypeName, const bool bAllowFallback = true) const
+	{
+		return GetInterfaceCustomizationByTypeId(GetElementTypeIdFromNameChecked(InElementTypeName), bAllowFallback);
+	}
+
+	/**
+	 * Get the interface customization for the given element type.
+	 * @note If bAllowFallback is true, then this will return the default interface customization if no override is present, otherwise it will return null.
+	 */
+	CustomizationBaseType* GetInterfaceCustomizationByTypeId(const FTypedHandleTypeId InElementTypeId, const bool bAllowFallback = true) const
+	{
+		CustomizationBaseType* InterfaceCustomization = RegisteredInterfaceCustomizations[InElementTypeId - 1].Get();
+		return InterfaceCustomization
+			? InterfaceCustomization
+			: bAllowFallback 
+				? DefaultInterfaceCustomization.Get()
+				: nullptr;
+	}
+
+private:
+	/** Default interface customization, used if no type-specific override is present. */
+	TUniquePtr<CustomizationBaseType> DefaultInterfaceCustomization;
+
+	/** Array of registered interface customizations, indexed by ElementTypeId-1. */
+	TUniquePtr<CustomizationBaseType> RegisteredInterfaceCustomizations[TypedHandleMaxTypeId - 1];
+};
diff --git a/Engine/Source/Runtime/TypedElementRuntime/Public/Elements/Framework/TypedElementListObjectUtil.h b/Engine/Source/Runtime/TypedElementRuntime/Public/Elements/Framework/TypedElementListObjectUtil.h
new file mode 100644
--- /dev/null
+++ b/Engine/Source/Runtime/TypedElementRuntime/Public/Elements/Framework/TypedElementListObjectUtil.h
@@ -0,0 +1,108 @@
+// Copyright Epic Games, Inc. All Rights Reserved.
+
+#pragma once
+
+#include "Elements/Framework/TypedElementList.h"
+
+namespace TypedElementListObjectUtil
+{
+
+/**
+ * Test whether there are any objects in the given list of elements.
+ */
+TYPEDELEMENTRUNTIME_API bool HasObjects(const UTypedElementList* InElementList, const UClass* InRequiredClass = nullptr);
+
+/**
+ * Test whether there are any objects in the given list of elements.
+ */
+template <typename RequiredClassType>
+bool HasObjects(const UTypedElementList* InElementList)
+{
+	return HasObjects(InElementList, RequiredClassType::StaticClass());
+}
+
+/**
+ * Count the number of objects in the given list of elements.
+ */
+TYPEDELEMENTRUNTIME_API int32 CountObjects(const UTypedElementList* InElementList, const UClass* InRequiredClass = nullptr);
+
+/**
+ * Count the number of objects in the given list of elements.
+ */
+template <typename RequiredClassType>
+int32 CountObjects(const UTypedElementList* InElementList)
+{
+	return CountObjects(InElementList, RequiredClassType::StaticClass());
+}
+
+/**
+ * Enumerate the objects from the given list of elements.
+ * @note Return true from the callback to continue enumeration.
+ */
+TYPEDELEMENTRUNTIME_API void ForEachObject(const UTypedElementList* InElementList, TFunctionRef<bool(UObject*)> InCallback, const UClass* InRequiredClass = nullptr);
+
+/**
+ * Enumerate the objects from the given list of elements.
+ * @note Return true from the callback to continue enumeration.
+ */
+template <typename RequiredClassType>
+void ForEachObject(const UTypedElementList* InElementList, TFunctionRef<bool(RequiredClassType*)> InCallback)
+{
+	ForEachObject(InElementList, [&InCallback](UObject* InObject)
+	{
+		return InCallback(CastChecked<RequiredClassType>(InObject));
+	}, RequiredClassType::StaticClass());
+}
+
+/**
+ * Get the array of objects from the given list of elements.
+ */
+TYPEDELEMENTRUNTIME_API TArray<UObject*> GetObjects(const UTypedElementList* InElementList, const UClass* InRequiredClass = nullptr);
+
+/**
+ * Get the array of objects from the given list of elements.
+ */
+template <typename RequiredClassType>
+TArray<RequiredClassType*> GetObjects(const UTypedElementList* InElementList)
+{
+	TArray<RequiredClassType*> SelectedObjects;
+	SelectedObjects.Reserve(InElementList->Num());
+
+	ForEachObject<RequiredClassType>(InElementList, [&SelectedObjects](RequiredClassType* InObject)
+	{
+		SelectedObjects.Add(InObject);
+		return true;
+	});
+
+	return SelectedObjects;
+}
+
+/**
+ * Get the first object of the given type from the given list of elements.
+ */
+TYPEDELEMENTRUNTIME_API UObject* GetTopObject(const UTypedElementList* InElementList, const UClass* InRequiredClass = nullptr);
+
+/**
+ * Get the first object of the given type from the given list of elements.
+ */
+template <typename RequiredClassType>
+RequiredClassType* GetTopObject(const UTypedElementList* InElementList)
+{
+	return Cast<RequiredClassType>(GetTopObject(InElementList, RequiredClassType::StaticClass()));
+}
+
+/**
+ * Get the last object of the given type from the given list of elements.
+ */
+TYPEDELEMENTRUNTIME_API UObject* GetBottomObject(const UTypedElementList* InElementList, const UClass* InRequiredClass = nullptr);
+
+/**
+ * Get the last object of the given type from the given list of elements.
+ */
+template <typename RequiredClassType>
+RequiredClassType* GetBottomObject(const UTypedElementList* InElementList)
+{
+	return Cast<RequiredClassType>(GetBottomObject(InElementList, RequiredClassType::StaticClass()));
+}
+
+} // namespace TypedElementListObjectUtil
diff --git a/Engine/Source/Runtime/TypedElementRuntime/Public/Elements/Framework/TypedElementSelectionSet.h b/Engine/Source/Runtime/TypedElementRuntime/Public/Elements/Framework/TypedElementSelectionSet.h
new file mode 100644
--- /dev/null
+++ b/Engine/Source/Runtime/TypedElementRuntime/Public/Elements/Framework/TypedElementSelectionSet.h
@@ -0,0 +1,453 @@
+// Copyright Epic Games, Inc. All Rights Reserved.
+
+#pragma once
+
+#include "Elements/Framework/TypedElementList.h"
+#include "Elements/Framework/TypedElementListObjectUtil.h"
+#include "Elements/Framework/TypedElementInterfaceCustomization.h"
+#include "Elements/Interfaces/TypedElementSelectionInterface.h"
+#include "TypedElementSelectionSet.generated.h"
+
+/**
+ * Customization type used to allow asset editors (such as the level editor) to override the base behavior of element selection,
+ * by injecting extra pre/post selection logic around the call into the selection interface for an element type.
+ */
+class TYPEDELEMENTRUNTIME_API FTypedElementSelectionCustomization
+{
+public:
+	virtual ~FTypedElementSelectionCustomization() = default;
+
+	//~ See UTypedElementSelectionInterface for API docs
+	virtual bool IsElementSelected(const TTypedElement<UTypedElementSelectionInterface>& InElementSelectionHandle, const UTypedElementList* InSelectionSet, const FTypedElementIsSelectedOptions& InSelectionOptions) { return InElementSelectionHandle.IsElementSelected(InSelectionSet, InSelectionOptions); }
+	virtual bool CanSelectElement(const TTypedElement<UTypedElementSelectionInterface>& InElementSelectionHandle, const FTypedElementSelectionOptions& InSelectionOptions) { return InElementSelectionHandle.CanSelectElement(InSelectionOptions); }
+	virtual bool CanDeselectElement(const TTypedElement<UTypedElementSelectionInterface>& InElementSelectionHandle, const FTypedElementSelectionOptions& InSelectionOptions) { return InElementSelectionHandle.CanDeselectElement(InSelectionOptions); }
+	virtual bool SelectElement(const TTypedElement<UTypedElementSelectionInterface>& InElementSelectionHandle, UTypedElementList* InSelectionSet, const FTypedElementSelectionOptions& InSelectionOptions) { return InElementSelectionHandle.SelectElement(InSelectionSet, InSelectionOptions); }
+	virtual bool DeselectElement(const TTypedElement<UTypedElementSelectionInterface>& InElementSelectionHandle, UTypedElementList* InSelectionSet, const FTypedElementSelectionOptions& InSelectionOptions) { return InElementSelectionHandle.DeselectElement(InSelectionSet, InSelectionOptions); }
+	virtual bool AllowSelectionModifiers(const TTypedElement<UTypedElementSelectionInterface>& InElementSelectionHandle, const UTypedElementList* InSelectionSet) { return InElementSelectionHandle.AllowSelectionModifiers(InSelectionSet); }
+	virtual FTypedElementHandle GetSelectionElement(const TTypedElement<UTypedElementSelectionInterface>& InElementSelectionHandle, const UTypedElementList* InCurrentSelection, const ETypedElementSelectionMethod InSelectionMethod) { return InElementSelectionHandle.GetSelectionElement(InCurrentSelection, InSelectionMethod); }
+};
+
+/**
+ * Utility to hold a typed element handle and its associated selection interface and selection customization.
+ */
+struct TYPEDELEMENTRUNTIME_API FTypedElementSelectionSetElement
+{
+public:
+	FTypedElementSelectionSetElement() = default;
+
+	FTypedElementSelectionSetElement(TTypedElement<UTypedElementSelectionInterface> InElementSelectionHandle, UTypedElementList* InElementList, FTypedElementSelectionCustomization* InSelectionCustomization)
+		: ElementSelectionHandle(MoveTemp(InElementSelectionHandle))
+		, ElementList(InElementList)
+		, SelectionCustomization(InSelectionCustomization)
+	{
+	}
+
+	FTypedElementSelectionSetElement(const FTypedElementSelectionSetElement&) = default;
+	FTypedElementSelectionSetElement& operator=(const FTypedElementSelectionSetElement&) = default;
+
+	FTypedElementSelectionSetElement(FTypedElementSelectionSetElement&&) = default;
+	FTypedElementSelectionSetElement& operator=(FTypedElementSelectionSetElement&&) = default;
+
+	FORCEINLINE explicit operator bool() const
+	{
+		return IsSet();
+	}
+
+	FORCEINLINE bool IsSet() const
+	{
+		return ElementSelectionHandle.IsSet()
+			&& ElementList
+			&& SelectionCustomization;
+	}
+
+	//~ See UTypedElementSelectionInterface for API docs
+	bool IsElementSelected(const FTypedElementIsSelectedOptions& InSelectionOptions) const { return SelectionCustomization->IsElementSelected(ElementSelectionHandle, ElementList, InSelectionOptions); }
+	bool CanSelectElement(const FTypedElementSelectionOptions& InSelectionOptions) const { return SelectionCustomization->CanSelectElement(ElementSelectionHandle, InSelectionOptions); }
+	bool CanDeselectElement(const FTypedElementSelectionOptions& InSelectionOptions) const { return SelectionCustomization->CanDeselectElement(ElementSelectionHandle, InSelectionOptions); }
+	bool SelectElement(const FTypedElementSelectionOptions& InSelectionOptions) const { return SelectionCustomization->SelectElement(ElementSelectionHandle, ElementList, InSelectionOptions); }
+	bool DeselectElement(const FTypedElementSelectionOptions& InSelectionOptions) const { return SelectionCustomization->DeselectElement(ElementSelectionHandle, ElementList, InSelectionOptions); }
+	bool AllowSelectionModifiers() const { return SelectionCustomization->AllowSelectionModifiers(ElementSelectionHandle, ElementList); }
+	FTypedElementHandle GetSelectionElement(const ETypedElementSelectionMethod InSelectionMethod) const { return SelectionCustomization->GetSelectionElement(ElementSelectionHandle, ElementList, InSelectionMethod); }
+
+private:
+	TTypedElement<UTypedElementSelectionInterface> ElementSelectionHandle;
+	UTypedElementList* ElementList = nullptr;
+	FTypedElementSelectionCustomization* SelectionCustomization = nullptr;
+};
+
+/**
+ * A wrapper around an element list that ensures mutation goes via the selection 
+ * interfaces, as well as providing some utilities for batching operations.
+ */
+UCLASS(Transient)
+class TYPEDELEMENTRUNTIME_API UTypedElementSelectionSet : public UObject, public TTypedElementInterfaceCustomizationRegistry<FTypedElementSelectionCustomization>
+{
+	GENERATED_BODY()
+
+public:
+	UTypedElementSelectionSet();
+
+	//~ UObject interface
+#if WITH_EDITOR
+	virtual bool Modify(bool bAlwaysMarkDirty = true) override;
+#endif	// WITH_EDITOR
+	virtual void Serialize(FArchive& Ar) override;
+
+	/**
+	 * Test to see whether the given element is currently considered selected.
+	 */
+	UFUNCTION(BlueprintPure, Category="TypedElementFramework|Selection")
+	bool IsElementSelected(const FTypedElementHandle& InElementHandle, const FTypedElementIsSelectedOptions InSelectionOptions) const;
+
+	/**
+	 * Test to see whether the given element can be selected.
+	 */
+	UFUNCTION(BlueprintPure, Category="TypedElementFramework|Selection")
+	bool CanSelectElement(const FTypedElementHandle& InElementHandle, const FTypedElementSelectionOptions InSelectionOptions) const;
+
+	/**
+	 * Test to see whether the given element can be deselected.
+	 */
+	UFUNCTION(BlueprintPure, Category="TypedElementFramework|Selection")
+	bool CanDeselectElement(const FTypedElementHandle& InElementHandle, const FTypedElementSelectionOptions InSelectionOptions) const;
+
+	/**
+	 * Attempt to select the given element.
+	 * @return True if the selection was changed, false otherwise.
+	 */
+	UFUNCTION(BlueprintCallable, Category="TypedElementFramework|Selection")
+	bool SelectElement(const FTypedElementHandle& InElementHandle, const FTypedElementSelectionOptions InSelectionOptions);
+
+	/**
+	 * Attempt to select the given elements.
+	 * @return True if the selection was changed, false otherwise.
+	 */
+	UFUNCTION(BlueprintCallable, Category="TypedElementFramework|Selection")
+	bool SelectElements(const TArray<FTypedElementHandle>& InElementHandles, const FTypedElementSelectionOptions InSelectionOptions);
+	bool SelectElements(TArrayView<const FTypedElementHandle> InElementHandles, const FTypedElementSelectionOptions InSelectionOptions);
+
+	/**
+	 * Attempt to deselect the given element.
+	 * @return True if the selection was changed, false otherwise.
+	 */
+	UFUNCTION(BlueprintCallable, Category="TypedElementFramework|Selection")
+	bool DeselectElement(const FTypedElementHandle& InElementHandle, const FTypedElementSelectionOptions InSelectionOptions);
+
+	/**
+	 * Attempt to deselect the given elements.
+	 * @return True if the selection was changed, false otherwise.
+	 */
+	UFUNCTION(BlueprintCallable, Category="TypedElementFramework|Selection")
+	bool DeselectElements(const TArray<FTypedElementHandle>& InElementHandles, const FTypedElementSelectionOptions InSelectionOptions);
+	bool DeselectElements(TArrayView<const FTypedElementHandle> InElementHandles, const FTypedElementSelectionOptions InSelectionOptions);
+
+	/**
+	 * Clear the current selection.
+	 * @return True if the selection was changed, false otherwise.
+	 */
+	UFUNCTION(BlueprintCallable, Category="TypedElementFramework|Selection")
+	bool ClearSelection(const FTypedElementSelectionOptions InSelectionOptions);
+
+	/**
+	 * Attempt to make the selection the given elements.
+	 * @note Equivalent to calling ClearSelection then SelectElements, but happens in a single batch.
+	 * @return True if the selection was changed, false otherwise.
+	 */
+	UFUNCTION(BlueprintCallable, Category="TypedElementFramework|Selection")
+	bool SetSelection(const TArray<FTypedElementHandle>& InElementHandles, const FTypedElementSelectionOptions InSelectionOptions);
+	bool SetSelection(TArrayView<const FTypedElementHandle> InElementHandles, const FTypedElementSelectionOptions InSelectionOptions);
+
+	/**
+	 * Test to see whether selection modifiers (Ctrl or Shift) are allowed while selecting this element.
+	 */
+	UFUNCTION(BlueprintPure, Category="TypedElementFramework|Selection")
+	bool AllowSelectionModifiers(const FTypedElementHandle& InElementHandle) const;
+
+	/**
+	 * Given an element, return the element that should actually perform a selection operation.
+	 */
+	UFUNCTION(BlueprintPure, Category="TypedElementFramework|Selection")
+	FTypedElementHandle GetSelectionElement(const FTypedElementHandle& InElementHandle, const ETypedElementSelectionMethod InSelectionMethod) const;
+
+	/**
+	 * Get the number of selected elements.
+	 */
+	UFUNCTION(BlueprintPure, Category="TypedElementFramework|Selection")
+	int32 GetNumSelectedElements() const
+	{
+		return ElementList->Num();
+	}
+
+	/**
+	 * Test whether there selected elements, optionally filtering to elements that implement the given interface.
+	 */
+	UFUNCTION(BlueprintCallable, BlueprintPure=false, Category="TypedElementFramework|Selection")
+	bool HasSelectedElements(const TSubclassOf<UTypedElementInterface>& InBaseInterfaceType = nullptr) const
+	{
+		return ElementList->HasElements(InBaseInterfaceType);
+	}
+
+	/**
+	 * Count the number of selected elements, optionally filtering to elements that implement the given interface.
+	 */
+	UFUNCTION(BlueprintCallable, BlueprintPure=false, Category="TypedElementFramework|Selection")
+	int32 CountSelectedElements(const TSubclassOf<UTypedElementInterface>& InBaseInterfaceType = nullptr) const
+	{
+		return ElementList->CountElements(InBaseInterfaceType);
+	}
+
+	/**
+	 * Get the handle of every selected element, optionally filtering to elements that implement the given interface.
+	 */
+	UFUNCTION(BlueprintCallable, BlueprintPure=false, Category="TypedElementFramework|Selection")
+	TArray<FTypedElementHandle> GetSelectedElementHandles(const TSubclassOf<UTypedElementInterface>& InBaseInterfaceType = nullptr) const
+	{
+		return ElementList->GetElementHandles(InBaseInterfaceType);
+	}
+
+	/**
+	 * Get the handle of every selected element, optionally filtering to elements that implement the given interface.
+	 */
+	template <typename ArrayAllocator>
+	void GetSelectedElementHandles(TArray<FTypedElementHandle, ArrayAllocator>& OutArray, const TSubclassOf<UTypedElementInterface>& InBaseInterfaceType = nullptr) const
+	{
+		ElementList->GetElementHandles(OutArray, InBaseInterfaceType);
+	}
+
+	/**
+	 * Enumerate the handle of every selected element, optionally filtering to elements that implement the given interface.
+	 * @note Return true from the callback to continue enumeration.
+	 */
+	void ForEachSelectedElementHandle(TFunctionRef<bool(const FTypedElementHandle&)> InCallback, const TSubclassOf<UTypedElementInterface>& InBaseInterfaceType = nullptr) const
+	{
+		ElementList->ForEachElementHandle(InCallback, InBaseInterfaceType);
+	}
+
+	/**
+	 * Enumerate the selected elements that implement the given interface.
+	 * @note Return true from the callback to continue enumeration.
+	 */
+	template <typename BaseInterfaceType>
+	void ForEachSelectedElement(TFunctionRef<bool(const TTypedElement<BaseInterfaceType>&)> InCallback) const
+	{
+		ElementList->ForEachElement<BaseInterfaceType>(InCallback);
+	}
+
+	/**
+	 * Get the first selected element implementing the given interface.
+	 */
+	template <typename BaseInterfaceType>
+	TTypedElement<BaseInterfaceType> GetTopSelectedElement() const
+	{
+		return ElementList->GetTopElement<BaseInterfaceType>();
+	}
+
+	/**
+	 * Get the last selected element implementing the given interface.
+	 */
+	template <typename BaseInterfaceType>
+	TTypedElement<BaseInterfaceType> GetBottomSelectedElement() const
+	{
+		return ElementList->GetBottomElement<BaseInterfaceType>();
+	}
+
+	/**
+	 * Test whether there are any selected objects.
+	 */
+	UFUNCTION(BlueprintCallable, BlueprintPure=false, Category="TypedElementFramework|Selection")
+	bool HasSelectedObjects(const UClass* InRequiredClass = nullptr) const
+	{
+		return TypedElementListObjectUtil::HasObjects(ElementList, InRequiredClass);
+	}
+
+	/**
+	 * Test whether there are any selected objects.
+	 */
+	template <typename RequiredClassType>
+	bool HasSelectedObjects() const
+	{
+		return TypedElementListObjectUtil::HasObjects<RequiredClassType>(ElementList);
+	}
+
+	/**
+	 * Count the number of selected objects.
+	 */
+	UFUNCTION(BlueprintCallable, BlueprintPure=false, Category="TypedElementFramework|Selection")
+	int32 CountSelectedObjects(const UClass* InRequiredClass = nullptr) const
+	{
+		return TypedElementListObjectUtil::CountObjects(ElementList, InRequiredClass);
+	}
+
+	/**
+	 * Count the number of selected objects.
+	 */
+	template <typename RequiredClassType>
+	int32 CountSelectedObjects() const
+	{
+		return TypedElementListObjectUtil::CountObjects<RequiredClassType>(ElementList);
+	}
+
+	/**
+	 * Get the array of selected objects from the currently selected elements.
+	 */
+	UFUNCTION(BlueprintCallable, BlueprintPure=false, Category="TypedElementFramework|Selection")
+	TArray<UObject*> GetSelectedObjects(const UClass* InRequiredClass = nullptr) const
+	{
+		return TypedElementListObjectUtil::GetObjects(ElementList, InRequiredClass);
+	}
+
+	/**
+	 * Get the array of selected objects from the currently selected elements.
+	 */
+	template <typename RequiredClassType>
+	TArray<RequiredClassType*> GetSelectedObjects() const
+	{
+		return TypedElementListObjectUtil::GetObjects<RequiredClassType>(ElementList);
+	}
+
+	/**
+	 * Enumerate the selected objects from the currently selected elements.
+	 * @note Return true from the callback to continue enumeration.
+	 */
+	void ForEachSelectedObject(TFunctionRef<bool(UObject*)> InCallback, const UClass* InRequiredClass = nullptr) const
+	{
+		TypedElementListObjectUtil::ForEachObject(ElementList, InCallback, InRequiredClass);
+	}
+
+	/**
+	 * Enumerate the selected objects from the currently selected elements.
+	 * @note Return true from the callback to continue enumeration.
+	 */
+	template <typename RequiredClassType>
+	void ForEachSelectedObject(TFunctionRef<bool(RequiredClassType*)> InCallback) const
+	{
+		TypedElementListObjectUtil::ForEachObject<RequiredClassType>(ElementList, InCallback);
+	}
+
+	/**
+	 * Get the first selected object of the given type.
+	 */
+	UFUNCTION(BlueprintPure, Category="TypedElementFramework|Selection")
+	UObject* GetTopSelectedObject(const UClass* InRequiredClass = nullptr) const
+	{
+		return TypedElementListObjectUtil::GetTopObject(ElementList, InRequiredClass);
+	}
+
+	/**
+	 * Get the first selected object of the given type.
+	 */
+	template <typename RequiredClassType>
+	RequiredClassType* GetTopSelectedObject() const
+	{
+		return TypedElementListObjectUtil::GetTopObject<RequiredClassType>(ElementList);
+	}
+
+	/**
+	 * Get the last selected object of the given type.
+	 */
+	UFUNCTION(BlueprintPure, Category="TypedElementFramework|Selection")
+	UObject* GetBottomSelectedObject(const UClass* InRequiredClass = nullptr) const
+	{
+		return TypedElementListObjectUtil::GetBottomObject(ElementList, InRequiredClass);
+	}
+
+	/**
+	 * Get the last selected object of the given type.
+	 */
+	template <typename RequiredClassType>
+	RequiredClassType* GetBottomSelectedObject() const
+	{
+		return TypedElementListObjectUtil::GetBottomObject<RequiredClassType>(ElementList);
+	}
+
+	/**
+	 * Access the delegate that is invoked whenever this element list is potentially about to change.
+	 * @note This may be called even if no actual change happens, though once a change does happen it won't be called again until after the next call to NotifyPendingChanges.
+	 */
+	DECLARE_EVENT_OneParam(UTypedElementSelectionSet, FOnPreChange, const UTypedElementSelectionSet* /*InElementSelectionSet*/);
+	FOnPreChange& OnPreChange()
+	{
+		return OnPreChangeDelegate;
+	}
+
+	/**
+	 * Access the delegate that is invoked whenever the underlying element list has been changed.
+	 * @note This is called automatically at the end of each frame, but can also be manually invoked by NotifyPendingChanges.
+	 */
+	DECLARE_EVENT_OneParam(UTypedElementSelectionSet, FOnChanged, const UTypedElementSelectionSet* /*InElementSelectionSet*/);
+	FOnChanged& OnChanged()
+	{
+		return OnChangedDelegate;
+	}
+
+	/**
+	 * Invoke the delegate called whenever the underlying element list has been changed.
+	 */
+	void NotifyPendingChanges()
+	{
+		ElementList->NotifyPendingChanges();
+	}
+
+	/**
+	 * Clear whether there are pending changes for OnChangedDelegate to notify for, without emitting a notification.
+	 */
+	void ClearPendingChanges()
+	{
+		ElementList->ClearPendingChanges();
+	}
+
+	/**
+	 * Access the interface to allow external systems (such as USelection) to receive immediate sync notifications as the underlying element list is changed.
+	 * This exists purely as a bridging mechanism and shouldn't be relied on for new code. It is lazily created as needed.
+	 */
+	FTypedElementListLegacySync& Legacy_GetElementListSync()
+	{
+		return ElementList->Legacy_GetSync();
+	}
+
+	/**
+	 * Access the interface to allow external systems (such as USelection) to receive immediate sync notifications as the underlying element list is changed.
+	 * This exists purely as a bridging mechanism and shouldn't be relied on for new code. This will return null if no legacy sync has been created for this instance.
+	 */
+	FTypedElementListLegacySync* Legacy_GetElementListSyncPtr()
+	{
+		return ElementList->Legacy_GetSyncPtr();
+	}
+
+	/**
+	 * Get the underlying element list holding the selection state.
+	 */
+	const UTypedElementList* GetElementList() const
+	{
+		return ElementList;
+	}
+
+private:
+	/**
+	 * Attempt to resolve the selection interface and selection customization for the given element, if any.
+	 */
+	FTypedElementSelectionSetElement ResolveSelectionSetElement(const FTypedElementHandle& InElementHandle) const;
+
+	/**
+	 * Proxy the internal OnPreChange event from the underlying element list.
+	 */
+	void OnElementListPreChange(const UTypedElementList* InElementList);
+
+	/**
+	 * Proxy the internal OnChanged event from the underlying element list.
+	 */
+	void OnElementListChanged(const UTypedElementList* InElementList);
+
+	/** Underlying element list holding the selection state. */
+	UPROPERTY()
+	UTypedElementList* ElementList = nullptr;
+
+	/** Delegate that is invoked whenever the underlying element list is potentially about to change. */
+	FOnPreChange OnPreChangeDelegate;
+
+	/** Delegate that is invoked whenever the underlying element list has been changed. */
+	FOnChanged OnChangedDelegate;
+
+	/** Set when we are currently restoring the selection state from a transaction (undo/redo) */
+	bool bIsRestoringFromTransaction = false;
+};
diff --git a/Engine/Source/Runtime/TypedElementRuntime/Public/Elements/Interfaces/TypedElementAssetDataInterface.h b/Engine/Source/Runtime/TypedElementRuntime/Public/Elements/Interfaces/TypedElementAssetDataInterface.h
new file mode 100644
--- /dev/null
+++ b/Engine/Source/Runtime/TypedElementRuntime/Public/Elements/Interfaces/TypedElementAssetDataInterface.h
@@ -0,0 +1,37 @@
+// Copyright Epic Games, Inc. All Rights Reserved.
+
+#pragma once
+
+#include "Elements/Framework/TypedElementHandle.h"
+#include "AssetRegistry/AssetData.h"
+
+#include "TypedElementAssetDataInterface.generated.h"
+
+UCLASS(Abstract)
+class TYPEDELEMENTRUNTIME_API UTypedElementAssetDataInterface : public UTypedElementInterface
+{
+	GENERATED_BODY()
+
+public:
+	/**
+	 * Returns any asset datas for content objects referenced by handle.
+	 * If the given handle itself has valid asset data, it should be returned as the last element of the array.
+	 *
+	 * @returns An array of valid asset datas.
+	 */
+	UFUNCTION(BlueprintPure, Category = "TypedElementInterfaces|AssetData")
+	virtual TArray<FAssetData> GetAllReferencedAssetDatas(const FTypedElementHandle& InElementHandle);
+
+	/**
+	 * Returns the asset data for the given handle, if it exists.
+	 */
+	UFUNCTION(BlueprintPure, Category = "TypedElementInterfaces|AssetData")
+	virtual FAssetData GetAssetData(const FTypedElementHandle& InElementHandle);
+};
+
+template <>
+struct TTypedElement<UTypedElementAssetDataInterface> : public TTypedElementBase<UTypedElementAssetDataInterface>
+{
+	TArray<FAssetData> GetAllReferencedAssetDatas() const { return InterfacePtr->GetAllReferencedAssetDatas(*this); }
+	FAssetData GetAssetData() { return InterfacePtr->GetAssetData(*this); }
+};
diff --git a/Engine/Source/Runtime/TypedElementRuntime/Public/Elements/Interfaces/TypedElementObjectInterface.h b/Engine/Source/Runtime/TypedElementRuntime/Public/Elements/Interfaces/TypedElementObjectInterface.h
new file mode 100644
--- /dev/null
+++ b/Engine/Source/Runtime/TypedElementRuntime/Public/Elements/Interfaces/TypedElementObjectInterface.h
@@ -0,0 +1,72 @@
+// Copyright Epic Games, Inc. All Rights Reserved.
+
+#pragma once
+
+#include "Elements/Framework/TypedElementHandle.h"
+#include "TypedElementObjectInterface.generated.h"
+
+UCLASS(Abstract)
+class TYPEDELEMENTRUNTIME_API UTypedElementObjectInterface : public UTypedElementInterface
+{
+	GENERATED_BODY()
+
+public:
+	/**
+	 * Get the object instance that this handle represents, if any.
+	 */
+	UFUNCTION(BlueprintPure, Category="TypedElementInterfaces|Object")
+	virtual UObject* GetObject(const FTypedElementHandle& InElementHandle);
+
+	/**
+	 * Gets the object instance's class that the handle represents, if any. 
+	 */
+	UFUNCTION(BlueprintPure, Category = "TypedElementInterfaces|Object")
+	UClass* GetObjectClass(const FTypedElementHandle& InElementHandle);
+
+	/**
+	 * Attempts to cast the given handle to another class.
+	 *
+	 * @returns the casted object if successful, otherwise nullptr.
+	 */
+	template <class CastTo>
+	CastTo* GetObjectAs(const FTypedElementHandle& InElementHandle)
+	{
+		return Cast<CastTo>(GetObject(InElementHandle));
+	}
+
+	/**
+	 * Attempts to cast the given handle to another class, if it can also be casted to TargetClass.
+	 * This is intended for use in cases where the calling code may only need an actor pointer, but also wants to be sure it's a specific type of actor.
+	 *
+	 * @returns the casted object if successful, otherwise nullptr.
+	 */
+	template <class CastTo = UObject>
+	CastTo* GetObjectAs(const FTypedElementHandle& InElementHandle, TSubclassOf<CastTo> TargetClass)
+	{
+		if (!TargetClass)
+		{
+			return nullptr;
+		}
+
+		UObject* ObjectPtr = GetObject(InElementHandle);
+		if (ObjectPtr && ObjectPtr->IsA(TargetClass))
+		{
+			return (CastTo*)ObjectPtr;
+		}
+
+		return nullptr;
+	}
+};
+
+template <>
+struct TTypedElement<UTypedElementObjectInterface> : public TTypedElementBase<UTypedElementObjectInterface>
+{
+	UObject* GetObject() const { return InterfacePtr->GetObject(*this); }
+	UClass* GetObjectClass() { return InterfacePtr->GetObjectClass(*this); }
+
+	template <class CastTo>
+	CastTo* GetObjectAs() { return InterfacePtr->GetObjectAs<CastTo>(*this); }
+
+	template <class CastTo>
+	CastTo* GetObjectAs(TSubclassOf<CastTo> TargetClass) { return InterfacePtr->GetObjectAs<CastTo>(*this, TargetClass); }
+};
diff --git a/Engine/Source/Runtime/TypedElementRuntime/Public/Elements/Interfaces/TypedElementSelectionInterface.h b/Engine/Source/Runtime/TypedElementRuntime/Public/Elements/Interfaces/TypedElementSelectionInterface.h
new file mode 100644
--- /dev/null
+++ b/Engine/Source/Runtime/TypedElementRuntime/Public/Elements/Interfaces/TypedElementSelectionInterface.h
@@ -0,0 +1,146 @@
+// Copyright Epic Games, Inc. All Rights Reserved.
+
+#pragma once
+
+#include "Elements/Framework/TypedElementHandle.h"
+#include "TypedElementSelectionInterface.generated.h"
+
+class UTypedElementList;
+
+UENUM()
+enum class ETypedElementSelectionMethod : uint8
+{
+	/** Select the "primary" element (eg, a component favor selecting its owner actor) */
+	Primary,
+	/** Select the "secondary" element (eg, a component would favor selecting itself) */
+	Secondary,
+};
+
+USTRUCT(BlueprintType)
+struct FTypedElementIsSelectedOptions
+{
+	GENERATED_BODY()
+
+public:
+	FTypedElementIsSelectedOptions& SetAllowIndirect(const bool InAllowIndirect) { bAllowIndirect = InAllowIndirect; return *this; }
+	bool AllowIndirect() const { return bAllowIndirect; }
+	
+private:
+	UPROPERTY(BlueprintReadWrite, Category="TypedElementInterfaces|Selection|IsSelectedOptions", meta=(AllowPrivateAccess=true))
+	bool bAllowIndirect = false;
+};
+
+USTRUCT(BlueprintType)
+struct FTypedElementSelectionOptions
+{
+	GENERATED_BODY()
+
+public:
+	FTypedElementSelectionOptions& SetAllowHidden(const bool InAllowHidden) { bAllowHidden = InAllowHidden; return *this; }
+	bool AllowHidden() const { return bAllowHidden; }
+
+	FTypedElementSelectionOptions& SetAllowGroups(const bool InAllowGroups) { bAllowGroups = InAllowGroups; return *this; }
+	bool AllowGroups() const { return bAllowGroups; }
+
+	FTypedElementSelectionOptions& SetAllowLegacyNotifications(const bool InAllowLegacyNotifications) { bAllowLegacyNotifications = InAllowLegacyNotifications; return *this; }
+	bool AllowLegacyNotifications() const { return bAllowLegacyNotifications; }
+
+	FTypedElementSelectionOptions& SetWarnIfLocked(const bool InWarnIfLocked) { bWarnIfLocked = InWarnIfLocked; return *this; }
+	bool WarnIfLocked() const { return bWarnIfLocked; }
+	
+private:
+	UPROPERTY(BlueprintReadWrite, Category="TypedElementInterfaces|Selection|SelectionOptions", meta=(AllowPrivateAccess=true))
+	bool bAllowHidden = false;
+
+	UPROPERTY(BlueprintReadWrite, Category="TypedElementInterfaces|Selection|SelectionOptions", meta=(AllowPrivateAccess=true))
+	bool bAllowGroups = true;
+
+	UPROPERTY(BlueprintReadWrite, Category="TypedElementInterfaces|Selection|SelectionOptions", meta=(AllowPrivateAccess=true))
+	bool bAllowLegacyNotifications = true;
+
+	UPROPERTY(BlueprintReadWrite, Category="TypedElementInterfaces|Selection|SelectionOptions", meta=(AllowPrivateAccess=true))
+	bool bWarnIfLocked = false;
+};
+
+UCLASS(Abstract)
+class TYPEDELEMENTRUNTIME_API UTypedElementSelectionInterface : public UTypedElementInterface
+{
+	GENERATED_BODY()
+
+public:
+	/**
+	 * Test to see whether the given element is currently considered selected.
+	 */
+	UFUNCTION(BlueprintPure, Category="TypedElementInterfaces|Selection")
+	virtual bool IsElementSelected(const FTypedElementHandle& InElementHandle, const UTypedElementList* InSelectionSet, const FTypedElementIsSelectedOptions& InSelectionOptions);
+
+	/**
+	 * Test to see whether the given element can be selected.
+	 */
+	UFUNCTION(BlueprintPure, Category="TypedElementInterfaces|Selection")
+	virtual bool CanSelectElement(const FTypedElementHandle& InElementHandle, const FTypedElementSelectionOptions& InSelectionOptions) { return true; }
+
+	/**
+	 * Test to see whether the given element can be deselected.
+	 */
+	UFUNCTION(BlueprintPure, Category="TypedElementInterfaces|Selection")
+	virtual bool CanDeselectElement(const FTypedElementHandle& InElementHandle, const FTypedElementSelectionOptions& InSelectionOptions) { return true; }
+
+	/**
+	 * Attempt to select the given element.
+	 * @return True if the selection was changed, false otherwise.
+	 */
+	UFUNCTION(BlueprintCallable, Category="TypedElementInterfaces|Selection")
+	virtual bool SelectElement(const FTypedElementHandle& InElementHandle, UTypedElementList* InSelectionSet, const FTypedElementSelectionOptions& InSelectionOptions);
+
+	/**
+	 * Attempt to deselect the given element.
+	 * @return True if the selection was changed, false otherwise.
+	 */
+	UFUNCTION(BlueprintCallable, Category="TypedElementInterfaces|Selection")
+	virtual bool DeselectElement(const FTypedElementHandle& InElementHandle, UTypedElementList* InSelectionSet, const FTypedElementSelectionOptions& InSelectionOptions);
+
+	/**
+	 * Test to see whether selection modifiers (Ctrl or Shift) are allowed while selecting this element.
+	 */
+	UFUNCTION(BlueprintPure, Category="TypedElementInterfaces|Selection")
+	virtual bool AllowSelectionModifiers(const FTypedElementHandle& InElementHandle, const UTypedElementList* InSelectionSet) { return true; }
+
+	/**
+	 * Given an element, return the element that should actually perform a selection operation.
+	 */
+	UFUNCTION(BlueprintPure, Category="TypedElementInterfaces|Selection")
+	virtual FTypedElementHandle GetSelectionElement(const FTypedElementHandle& InElementHandle, const UTypedElementList* InCurrentSelection, const ETypedElementSelectionMethod InSelectionMethod) { return InElementHandle; }
+
+	/**
+	 * Test to see whether the given element prevents the selection set state from being transacted for undo/redo (eg, if the element belongs to a PIE instance).
+	 */
+	virtual bool ShouldPreventTransactions(const FTypedElementHandle& InElementHandle) { return false; }
+
+	/**
+	 * Write the information needed to find the given element again when this transaction is replayed for undo/redo.
+	 * @note The data written must match the data read by ReadTransactedElement.
+	 */
+	virtual void WriteTransactedElement(const FTypedElementHandle& InElementHandle, FArchive& InArchive) {}
+
+	/**
+	 * Read the information needed to be find a previously tracked element, and attempt to get it now that this transaction is being replayed for undo/redo.
+	 * @note The data read must match the data written by WriteTransactedElement.
+	 */
+	virtual FTypedElementHandle ReadTransactedElement(FArchive& InArchive) { return FTypedElementHandle(); }
+};
+
+template <>
+struct TTypedElement<UTypedElementSelectionInterface> : public TTypedElementBase<UTypedElementSelectionInterface>
+{
+	bool IsElementSelected(const UTypedElementList* InSelectionSet, const FTypedElementIsSelectedOptions& InSelectionOptions) const { return InterfacePtr->IsElementSelected(*this, InSelectionSet, InSelectionOptions); }
+	bool CanSelectElement(const FTypedElementSelectionOptions& InSelectionOptions) const { return InterfacePtr->CanSelectElement(*this, InSelectionOptions); }
+	bool CanDeselectElement(const FTypedElementSelectionOptions& InSelectionOptions) const { return InterfacePtr->CanDeselectElement(*this, InSelectionOptions); }
+	bool SelectElement(UTypedElementList* InSelectionSet, const FTypedElementSelectionOptions& InSelectionOptions) const { return InterfacePtr->SelectElement(*this, InSelectionSet, InSelectionOptions); }
+	bool DeselectElement(UTypedElementList* InSelectionSet, const FTypedElementSelectionOptions& InSelectionOptions) const { return InterfacePtr->DeselectElement(*this, InSelectionSet, InSelectionOptions); }
+	bool AllowSelectionModifiers(const UTypedElementList* InSelectionSet) const { return InterfacePtr->AllowSelectionModifiers(*this, InSelectionSet); }
+	FTypedElementHandle GetSelectionElement(const UTypedElementList* InCurrentSelection, const ETypedElementSelectionMethod InSelectionMethod) const { return InterfacePtr->GetSelectionElement(*this, InCurrentSelection, InSelectionMethod); }
+	bool ShouldPreventTransactions() const { return InterfacePtr->ShouldPreventTransactions(*this); }
+	void WriteTransactedElement(FArchive& InArchive) const { return InterfacePtr->WriteTransactedElement(*this, InArchive); }
+	FTypedElementHandle ReadTransactedElement(FArchive& InArchive) const { return InterfacePtr->ReadTransactedElement(InArchive); }
+};
diff --git a/Engine/Source/Runtime/TypedElementRuntime/TypedElementRuntime.Build.cs b/Engine/Source/Runtime/TypedElementRuntime/TypedElementRuntime.Build.cs
new file mode 100644
--- /dev/null
+++ b/Engine/Source/Runtime/TypedElementRuntime/TypedElementRuntime.Build.cs
@@ -0,0 +1,19 @@
+// Copyright Epic Games, Inc. All Rights Reserved.
+
+using UnrealBuildTool;
+using System.IO;
+
+public class TypedElementRuntime : ModuleRules
+{
+	public TypedElementRuntime(ReadOnlyTargetRules Target) : base(Target)
+	{
+		PublicDependencyModuleNames.AddRange(
+			new string[]
+			{
+				"Core",
+				"CoreUObject",
+				"TypedElementFramework",
+			}
+		);
+    }
+}
