Guarding against null ref exceptions that could be triggered when serializing the CDO of this class, which has no element list. This case originated when using a serializer to scan through all CDOs.

diff --git a/Engine/Source/Runtime/TypedElementRuntime/Private/Elements/Framework/TypedElementSelectionSet.cpp b/Engine/Source/Runtime/TypedElementRuntime/Private/Elements/Framework/TypedElementSelectionSet.cpp
--- a/Engine/Source/Runtime/TypedElementRuntime/Private/Elements/Framework/TypedElementSelectionSet.cpp
+++ b/Engine/Source/Runtime/TypedElementRuntime/Private/Elements/Framework/TypedElementSelectionSet.cpp
@@ -45,14 +45,17 @@ void UTypedElementSelectionSet::Serialize(FArchive& Ar)
 	{
 		FTypedHandleTypeId ElementTypeId = 0;
 
-		ElementList->ForEachElement<UTypedElementSelectionInterface>([&Ar, &ElementTypeId](const TTypedElement<UTypedElementSelectionInterface>& InSelectionElement)
+		if (ElementList)
 		{
-			ElementTypeId = InSelectionElement.GetId().GetTypeId();
-			Ar << ElementTypeId;
+			ElementList->ForEachElement<UTypedElementSelectionInterface>([&Ar, &ElementTypeId](const TTypedElement<UTypedElementSelectionInterface>& InSelectionElement)
+			{
+				ElementTypeId = InSelectionElement.GetId().GetTypeId();
+				Ar << ElementTypeId;
 
-			InSelectionElement.WriteTransactedElement(Ar);
-			return true;
-		});
+				InSelectionElement.WriteTransactedElement(Ar);
+				return true;
+			});
+		}
 
 		// End of the list
 		ElementTypeId = 0;
@@ -85,6 +88,7 @@ void UTypedElementSelectionSet::Serialize(FArchive& Ar)
 			}
 		}
 
+		if (ElementList)
 		{
 			FTypedElementListLegacySyncScopedBatch LegacySyncBatch(ElementList, /*bNotify*/false);
 			TGuardValue<bool> GuardIsRestoringFromTransaction(bIsRestoringFromTransaction, true);
