Deferred restoring the selected element state until PostEditUndo

This is required as the Serialize order of transacted objects is not guaranteed, and it may not be possible to resolve the elements used by the previous selection state until everything has been serialized.

This adds an intermediary type, ITypedElementTransactedElement, to hold the transacted state that used to be directly serialized, and updates FTypedElementSelectionSetState to use that as its stored state (rather than the serialized blob). Serialize is now implemented using GetCurrentSelectionState and RestoreSelectionState, rather than those functions being implemented using Serialize.

diff --git a/Engine/Source/Editor/UnrealEd/Private/Elements/Actor/ActorElementEditorSelectionInterface.cpp b/Engine/Source/Editor/UnrealEd/Private/Elements/Actor/ActorElementEditorSelectionInterface.cpp
--- a/Engine/Source/Editor/UnrealEd/Private/Elements/Actor/ActorElementEditorSelectionInterface.cpp
+++ b/Engine/Source/Editor/UnrealEd/Private/Elements/Actor/ActorElementEditorSelectionInterface.cpp
@@ -8,6 +8,36 @@
 #include "Elements/Framework/EngineElementsLibrary.h"
 #include "Elements/Object/ObjectElementEditorSelectionInterface.h"
 
+class FActorElementTransactedElement : public ITypedElementTransactedElement
+{
+private:
+	virtual TUniquePtr<ITypedElementTransactedElement> CloneImpl() const override
+	{
+		return MakeUnique<FActorElementTransactedElement>(*this);
+	}
+
+	virtual FTypedElementHandle GetElementImpl() const override
+	{
+		const AActor* Actor = ActorPtr.Get(/*bEvenIfPendingKill*/true);
+		return Actor
+			? UEngineElementsLibrary::AcquireEditorActorElementHandle(Actor)
+			: FTypedElementHandle();
+	}
+
+	virtual void SetElementImpl(const FTypedElementHandle& InElementHandle) override
+	{
+		AActor* Actor = ActorElementDataUtil::GetActorFromHandleChecked(InElementHandle);
+		ActorPtr = Actor;
+	}
+
+	virtual void SerializeImpl(FArchive& InArchive) override
+	{
+		InArchive << ActorPtr;
+	}
+
+	TWeakObjectPtr<AActor> ActorPtr;
+};
+
 bool UActorElementEditorSelectionInterface::IsElementSelected(const FTypedElementHandle& InElementHandle, const UTypedElementList* InSelectionSet, const FTypedElementIsSelectedOptions& InSelectionOptions)
 {
 	const AActor* Actor = ActorElementDataUtil::GetActorFromHandle(InElementHandle);
@@ -20,20 +50,9 @@ bool UActorElementEditorSelectionInterface::ShouldPreventTransactions(const FTyp
 	return Actor && UObjectElementEditorSelectionInterface::ShouldObjectPreventTransactions(Actor);
 }
 
-void UActorElementEditorSelectionInterface::WriteTransactedElement(const FTypedElementHandle& InElementHandle, FArchive& InArchive)
-{
-	if (const AActor* Actor = ActorElementDataUtil::GetActorFromHandle(InElementHandle))
-	{
-		UObjectElementEditorSelectionInterface::WriteTransactedObject(Actor, InArchive);
-	}
-}
-
-FTypedElementHandle UActorElementEditorSelectionInterface::ReadTransactedElement(FArchive& InArchive)
+TUniquePtr<ITypedElementTransactedElement> UActorElementEditorSelectionInterface::CreateTransactedElementImpl()
 {
-	return UObjectElementEditorSelectionInterface::ReadTransactedObject(InArchive, [](const UObject* InObject)
-	{
-		return UEngineElementsLibrary::AcquireEditorActorElementHandle(CastChecked<const AActor>(InObject));
-	});
+	return MakeUnique<FActorElementTransactedElement>();
 }
 
 bool UActorElementEditorSelectionInterface::IsActorSelected(const AActor* InActor, const UTypedElementList* InSelectionSet, const FTypedElementIsSelectedOptions& InSelectionOptions)
diff --git a/Engine/Source/Editor/UnrealEd/Private/Elements/Actor/ActorElementEditorSelectionInterface.h b/Engine/Source/Editor/UnrealEd/Private/Elements/Actor/ActorElementEditorSelectionInterface.h
--- a/Engine/Source/Editor/UnrealEd/Private/Elements/Actor/ActorElementEditorSelectionInterface.h
+++ b/Engine/Source/Editor/UnrealEd/Private/Elements/Actor/ActorElementEditorSelectionInterface.h
@@ -15,8 +15,7 @@ class UActorElementEditorSelectionInterface : public UActorElementSelectionInter
 public:
 	virtual bool IsElementSelected(const FTypedElementHandle& InElementHandle, const UTypedElementList* InSelectionSet, const FTypedElementIsSelectedOptions& InSelectionOptions) override;
 	virtual bool ShouldPreventTransactions(const FTypedElementHandle& InElementHandle) override;
-	virtual void WriteTransactedElement(const FTypedElementHandle& InElementHandle, FArchive& InArchive) override;
-	virtual FTypedElementHandle ReadTransactedElement(FArchive& InArchive) override;
+	virtual TUniquePtr<ITypedElementTransactedElement> CreateTransactedElementImpl() override;
 
 	static bool IsActorSelected(const AActor* InActor, const UTypedElementList* InSelectionSet, const FTypedElementIsSelectedOptions& InSelectionOptions);
 };
diff --git a/Engine/Source/Editor/UnrealEd/Private/Elements/Component/ComponentElementEditorSelectionInterface.cpp b/Engine/Source/Editor/UnrealEd/Private/Elements/Component/ComponentElementEditorSelectionInterface.cpp
--- a/Engine/Source/Editor/UnrealEd/Private/Elements/Component/ComponentElementEditorSelectionInterface.cpp
+++ b/Engine/Source/Editor/UnrealEd/Private/Elements/Component/ComponentElementEditorSelectionInterface.cpp
@@ -9,6 +9,36 @@
 #include "Elements/Framework/EngineElementsLibrary.h"
 #include "Elements/Object/ObjectElementEditorSelectionInterface.h"
 
+class FComponentElementTransactedElement : public ITypedElementTransactedElement
+{
+private:
+	virtual TUniquePtr<ITypedElementTransactedElement> CloneImpl() const override
+	{
+		return MakeUnique<FComponentElementTransactedElement>(*this);
+	}
+
+	virtual FTypedElementHandle GetElementImpl() const override
+	{
+		const UActorComponent* Component = ComponentPtr.Get(/*bEvenIfPendingKill*/true);
+		return Component
+			? UEngineElementsLibrary::AcquireEditorComponentElementHandle(Component)
+			: FTypedElementHandle();
+	}
+
+	virtual void SetElementImpl(const FTypedElementHandle& InElementHandle) override
+	{
+		UActorComponent* Component = ComponentElementDataUtil::GetComponentFromHandle(InElementHandle);
+		ComponentPtr = Component;
+	}
+
+	virtual void SerializeImpl(FArchive& InArchive) override
+	{
+		InArchive << ComponentPtr;
+	}
+
+	TWeakObjectPtr<UActorComponent> ComponentPtr;
+};
+
 bool UComponentElementEditorSelectionInterface::IsElementSelected(const FTypedElementHandle& InElementHandle, const UTypedElementList* InSelectionSet, const FTypedElementIsSelectedOptions& InSelectionOptions)
 {
 	const UActorComponent* Component = ComponentElementDataUtil::GetComponentFromHandle(InElementHandle);
@@ -21,20 +51,9 @@ bool UComponentElementEditorSelectionInterface::ShouldPreventTransactions(const
 	return Component && UObjectElementEditorSelectionInterface::ShouldObjectPreventTransactions(Component);
 }
 
-void UComponentElementEditorSelectionInterface::WriteTransactedElement(const FTypedElementHandle& InElementHandle, FArchive& InArchive)
-{
-	if (const UActorComponent* Component = ComponentElementDataUtil::GetComponentFromHandleChecked(InElementHandle))
-	{
-		UObjectElementEditorSelectionInterface::WriteTransactedObject(Component, InArchive);
-	}
-}
-
-FTypedElementHandle UComponentElementEditorSelectionInterface::ReadTransactedElement(FArchive& InArchive)
+TUniquePtr<ITypedElementTransactedElement> UComponentElementEditorSelectionInterface::CreateTransactedElementImpl()
 {
-	return UObjectElementEditorSelectionInterface::ReadTransactedObject(InArchive, [](const UObject* InObject)
-	{
-		return UEngineElementsLibrary::AcquireEditorComponentElementHandle(CastChecked<const UActorComponent>(InObject));
-	});
+	return MakeUnique<FComponentElementTransactedElement>();
 }
 
 bool UComponentElementEditorSelectionInterface::IsComponentSelected(const UActorComponent* InComponent, const UTypedElementList* InSelectionSet, const FTypedElementIsSelectedOptions& InSelectionOptions)
diff --git a/Engine/Source/Editor/UnrealEd/Private/Elements/Component/ComponentElementEditorSelectionInterface.h b/Engine/Source/Editor/UnrealEd/Private/Elements/Component/ComponentElementEditorSelectionInterface.h
--- a/Engine/Source/Editor/UnrealEd/Private/Elements/Component/ComponentElementEditorSelectionInterface.h
+++ b/Engine/Source/Editor/UnrealEd/Private/Elements/Component/ComponentElementEditorSelectionInterface.h
@@ -15,8 +15,7 @@ class UComponentElementEditorSelectionInterface : public UComponentElementSelect
 public:
 	virtual bool IsElementSelected(const FTypedElementHandle& InElementHandle, const UTypedElementList* InSelectionSet, const FTypedElementIsSelectedOptions& InSelectionOptions) override;
 	virtual bool ShouldPreventTransactions(const FTypedElementHandle& InElementHandle) override;
-	virtual void WriteTransactedElement(const FTypedElementHandle& InElementHandle, FArchive& InArchive) override;
-	virtual FTypedElementHandle ReadTransactedElement(FArchive& InArchive) override;
+	virtual TUniquePtr<ITypedElementTransactedElement> CreateTransactedElementImpl() override;
 
 	static bool IsComponentSelected(const UActorComponent* InComponent, const UTypedElementList* InSelectionSet, const FTypedElementIsSelectedOptions& InSelectionOptions);
 };
diff --git a/Engine/Source/Editor/UnrealEd/Private/Elements/Object/ObjectElementEditorSelectionInterface.cpp b/Engine/Source/Editor/UnrealEd/Private/Elements/Object/ObjectElementEditorSelectionInterface.cpp
--- a/Engine/Source/Editor/UnrealEd/Private/Elements/Object/ObjectElementEditorSelectionInterface.cpp
+++ b/Engine/Source/Editor/UnrealEd/Private/Elements/Object/ObjectElementEditorSelectionInterface.cpp
@@ -6,48 +6,49 @@
 #include "UObject/Package.h"
 #include "Elements/Framework/EngineElementsLibrary.h"
 
-bool UObjectElementEditorSelectionInterface::ShouldPreventTransactions(const FTypedElementHandle& InElementHandle)
+class FObjectElementTransactedElement : public ITypedElementTransactedElement
 {
-	const UObject* Object = ObjectElementDataUtil::GetObjectFromHandle(InElementHandle);
-	return Object && ShouldObjectPreventTransactions(Object);
-}
+private:
+	virtual TUniquePtr<ITypedElementTransactedElement> CloneImpl() const override
+	{
+		return MakeUnique<FObjectElementTransactedElement>(*this);
+	}
 
-void UObjectElementEditorSelectionInterface::WriteTransactedElement(const FTypedElementHandle& InElementHandle, FArchive& InArchive)
-{
-	if (const UObject* Object = ObjectElementDataUtil::GetObjectFromHandle(InElementHandle))
+	virtual FTypedElementHandle GetElementImpl() const override
 	{
-		WriteTransactedObject(Object, InArchive);
+		const UObject* Object = ObjectPtr.Get(/*bEvenIfPendingKill*/true);
+		return Object
+			? UEngineElementsLibrary::AcquireEditorObjectElementHandle(Object)
+			: FTypedElementHandle();
 	}
-}
 
-FTypedElementHandle UObjectElementEditorSelectionInterface::ReadTransactedElement(FArchive& InArchive)
-{
-	return ReadTransactedObject(InArchive, [](const UObject* InObject)
+	virtual void SetElementImpl(const FTypedElementHandle& InElementHandle) override
 	{
-		check(InObject);
-		return UEngineElementsLibrary::AcquireEditorObjectElementHandle(InObject);
-	});
-}
+		UObject* Object = ObjectElementDataUtil::GetObjectFromHandleChecked(InElementHandle);
+		ObjectPtr = Object;
+	}
 
-bool UObjectElementEditorSelectionInterface::ShouldObjectPreventTransactions(const UObject* InObject)
+	virtual void SerializeImpl(FArchive& InArchive) override
+	{
+		InArchive << ObjectPtr;
+	}
+
+	TWeakObjectPtr<UObject> ObjectPtr;
+};
+
+bool UObjectElementEditorSelectionInterface::ShouldPreventTransactions(const FTypedElementHandle& InElementHandle)
 {
-	// If the selection currently contains any PIE objects we should not be including it in the transaction buffer
-	return InObject->GetOutermost()->HasAnyPackageFlags(PKG_PlayInEditor | PKG_ContainsScript | PKG_CompiledIn);
+	const UObject* Object = ObjectElementDataUtil::GetObjectFromHandle(InElementHandle);
+	return Object && ShouldObjectPreventTransactions(Object);
 }
 
-void UObjectElementEditorSelectionInterface::WriteTransactedObject(const UObject* InObject, FArchive& InArchive)
+TUniquePtr<ITypedElementTransactedElement> UObjectElementEditorSelectionInterface::CreateTransactedElementImpl()
 {
-	TWeakObjectPtr<const UObject> SelectedObject = InObject;
-	InArchive << SelectedObject;
+	return MakeUnique<FObjectElementTransactedElement>();
 }
 
-FTypedElementHandle UObjectElementEditorSelectionInterface::ReadTransactedObject(FArchive& InArchive, TFunctionRef<FTypedElementHandle(const UObject*)> InFindElementForObject)
+bool UObjectElementEditorSelectionInterface::ShouldObjectPreventTransactions(const UObject* InObject)
 {
-	TWeakObjectPtr<const UObject> SelectedObject;
-	InArchive << SelectedObject;
-
-	const UObject* SelectedObjectPtr = SelectedObject.Get(/*bEvenIfPendingKill*/true);
-	return SelectedObjectPtr
-		? InFindElementForObject(SelectedObjectPtr)
-		: FTypedElementHandle();
+	// If the selection currently contains any PIE objects we should not be including it in the transaction buffer
+	return InObject->GetOutermost()->HasAnyPackageFlags(PKG_PlayInEditor | PKG_ContainsScript | PKG_CompiledIn);
 }
diff --git a/Engine/Source/Editor/UnrealEd/Private/Elements/Object/ObjectElementEditorSelectionInterface.h b/Engine/Source/Editor/UnrealEd/Private/Elements/Object/ObjectElementEditorSelectionInterface.h
--- a/Engine/Source/Editor/UnrealEd/Private/Elements/Object/ObjectElementEditorSelectionInterface.h
+++ b/Engine/Source/Editor/UnrealEd/Private/Elements/Object/ObjectElementEditorSelectionInterface.h
@@ -12,10 +12,7 @@ class UObjectElementEditorSelectionInterface : public UObjectElementSelectionInt
 
 public:
 	virtual bool ShouldPreventTransactions(const FTypedElementHandle& InElementHandle) override;
-	virtual void WriteTransactedElement(const FTypedElementHandle& InElementHandle, FArchive& InArchive) override;
-	virtual FTypedElementHandle ReadTransactedElement(FArchive& InArchive) override;
+	virtual TUniquePtr<ITypedElementTransactedElement> CreateTransactedElementImpl() override;
 
 	static bool ShouldObjectPreventTransactions(const UObject* InObject);
-	static void WriteTransactedObject(const UObject* InObject, FArchive& InArchive);
-	static FTypedElementHandle ReadTransactedObject(FArchive& InArchive, TFunctionRef<FTypedElementHandle(const UObject*)> InFindElementForObject);
 };
diff --git a/Engine/Source/Runtime/TypedElementRuntime/Private/Elements/Framework/TypedElementSelectionSet.cpp b/Engine/Source/Runtime/TypedElementRuntime/Private/Elements/Framework/TypedElementSelectionSet.cpp
--- a/Engine/Source/Runtime/TypedElementRuntime/Private/Elements/Framework/TypedElementSelectionSet.cpp
+++ b/Engine/Source/Runtime/TypedElementRuntime/Private/Elements/Framework/TypedElementSelectionSet.cpp
@@ -17,6 +17,23 @@ UTypedElementSelectionSet::UTypedElementSelectionSet()
 }
 
 #if WITH_EDITOR
+void UTypedElementSelectionSet::PreEditUndo()
+{
+	Super::PreEditUndo();
+
+	checkf(!PendingUndoRedoState, TEXT("PendingUndoRedoState was set! Missing call to PostEditUndo?"));
+	PendingUndoRedoState = MakeUnique<FTypedElementSelectionSetState>();
+}
+
+void UTypedElementSelectionSet::PostEditUndo()
+{
+	Super::PostEditUndo();
+
+	checkf(PendingUndoRedoState, TEXT("PendingUndoRedoState was null! Missing call to PreEditUndo?"));
+	RestoreSelectionState(*PendingUndoRedoState);
+	PendingUndoRedoState.Reset();
+}
+
 bool UTypedElementSelectionSet::Modify(bool bAlwaysMarkDirty)
 {
 	if (GUndo && CanModify())
@@ -46,65 +63,51 @@ void UTypedElementSelectionSet::Serialize(FArchive& Ar)
 
 	if (Ar.IsSaving())
 	{
-		FTypedHandleTypeId ElementTypeId = 0;
+		FTypedElementSelectionSetState SelectionState = ElementList ? GetCurrentSelectionState() : FTypedElementSelectionSetState();
+
+		int32 NumTransactedElements = SelectionState.TransactedElements.Num();
+		Ar << NumTransactedElements;
 
-		if (ElementList)
+		for (const TUniquePtr<ITypedElementTransactedElement>& TransactedElement : SelectionState.TransactedElements)
 		{
-			ElementList->ForEachElement<UTypedElementSelectionInterface>([&Ar, &ElementTypeId](const TTypedElement<UTypedElementSelectionInterface>& InSelectionElement)
-			{
-				ElementTypeId = InSelectionElement.GetId().GetTypeId();
-				Ar << ElementTypeId;
+			FTypedHandleTypeId TypeId = TransactedElement->GetElementType();
+			Ar << TypeId;
 
-				InSelectionElement.WriteTransactedElement(Ar);
-				return true;
-			});
+			TransactedElement->Serialize(Ar);
 		}
-
-		// End of the list
-		ElementTypeId = 0;
-		Ar << ElementTypeId;
 	}
 	else if (Ar.IsLoading())
 	{
-		TArray<FTypedElementHandle, TInlineAllocator<256>> SelectedElements;
+		UTypedElementRegistry* Registry = UTypedElementRegistry::GetInstance();
 
-		{
-			UTypedElementRegistry* Registry = UTypedElementRegistry::GetInstance();
+		const bool bIsUndoRedo = PendingUndoRedoState && Ar.IsTransacting();
+		FTypedElementSelectionSetState TmpSelectionState;
 
-			FTypedHandleTypeId ElementTypeId = 0;
-			for (;;)
-			{
-				Ar << ElementTypeId;
-				if (!ElementTypeId)
-				{
-					// End of the list
-					break;
-				}
-
-				UTypedElementSelectionInterface* ElementTypeSelectionInterface = Registry->GetElementInterface<UTypedElementSelectionInterface>(ElementTypeId);
-				checkf(ElementTypeSelectionInterface, TEXT("Failed to find selection interface for a previously transacted element type!"));
-				FTypedElementHandle SelectedElement = ElementTypeSelectionInterface->ReadTransactedElement(Ar);
-				if (SelectedElement)
-				{
-					SelectedElements.Add(MoveTemp(SelectedElement));
-				}
-			}
-		}
+		FTypedElementSelectionSetState& SelectionState = bIsUndoRedo ? *PendingUndoRedoState : TmpSelectionState;
+		SelectionState.CreatedFromSelectionSet = this;
+
+		int32 NumTransactedElements = 0;
+		Ar << NumTransactedElements;
 
-		if (ElementList)
+		SelectionState.TransactedElements.Reserve(NumTransactedElements);
+		for (int32 TransactedElementIndex = 0; TransactedElementIndex < NumTransactedElements; ++TransactedElementIndex)
 		{
-			FTypedElementListLegacySyncScopedBatch LegacySyncBatch(ElementList, /*bNotify*/false);
-			TGuardValue<bool> GuardIsRestoringFromTransaction(bIsRestoringFromTransaction, true);
+			FTypedHandleTypeId TypeId = 0;
+			Ar << TypeId;
 
-			// TODO: Work out the intersection of the before and after state instead of clearing and reselecting?
+			UTypedElementSelectionInterface* ElementTypeSelectionInterface = Registry->GetElementInterface<UTypedElementSelectionInterface>(TypeId);
+			checkf(ElementTypeSelectionInterface, TEXT("Failed to find selection interface for a previously transacted element type!"));
 
-			const FTypedElementSelectionOptions SelectionOptions = FTypedElementSelectionOptions()
-				.SetAllowHidden(true)
-				.SetAllowGroups(false)
-				.SetWarnIfLocked(false);
+			TUniquePtr<ITypedElementTransactedElement> TransactedElement = ElementTypeSelectionInterface->CreateTransactedElement(TypeId);
+			checkf(TransactedElement, TEXT("Failed to allocate a transacted element for a previously transacted element type!"));
+
+			TransactedElement->Serialize(Ar);
+			SelectionState.TransactedElements.Emplace(MoveTemp(TransactedElement));
+		}
 
-			ClearSelection(SelectionOptions);
-			SelectElements(SelectedElements, SelectionOptions);
+		if (ElementList && !bIsUndoRedo)
+		{
+			RestoreSelectionState(SelectionState);
 		}
 	}
 }
@@ -240,24 +243,48 @@ FTypedElementSelectionSetState UTypedElementSelectionSet::GetCurrentSelectionSta
 {
 	FTypedElementSelectionSetState CurrentState;
 
-	FObjectWriter TempArchive(const_cast<UTypedElementSelectionSet*>(this), CurrentState.StoredSelectionSetData);
-	if (TempArchive.IsError())
-	{
-		CurrentState.StoredSelectionSetData.Reset();
-	}
-	else
+	CurrentState.CreatedFromSelectionSet = this;
+
+	CurrentState.TransactedElements.Reserve(ElementList->Num());
+	ElementList->ForEachElement<UTypedElementSelectionInterface>([&CurrentState](const TTypedElement<UTypedElementSelectionInterface>& InSelectionElement)
 	{
-		CurrentState.CreatedFromSelectionSet = this;
-	}
+		if (TUniquePtr<ITypedElementTransactedElement> TransactedElement = InSelectionElement.CreateTransactedElement())
+		{
+			CurrentState.TransactedElements.Emplace(MoveTemp(TransactedElement));
+		}
+		return true;
+	});
 
-	return MoveTemp(CurrentState);
+	return CurrentState;
 }
 
 void UTypedElementSelectionSet::RestoreSelectionState(const FTypedElementSelectionSetState& InSelectionState)
 {
-	if ((InSelectionState.CreatedFromSelectionSet == this) && (InSelectionState.StoredSelectionSetData.Num() > 0))
+	if (InSelectionState.CreatedFromSelectionSet == this)
 	{
-		FObjectReader TempArchive(this, InSelectionState.StoredSelectionSetData);
+		TArray<FTypedElementHandle, TInlineAllocator<256>> SelectedElements;
+		SelectedElements.Reserve(InSelectionState.TransactedElements.Num());
+		
+		for (const TUniquePtr<ITypedElementTransactedElement>& TransactedElement : InSelectionState.TransactedElements)
+		{
+			if (FTypedElementHandle SelectedElement = TransactedElement->GetElement())
+			{
+				SelectedElements.Add(MoveTemp(SelectedElement));
+			}
+		}
+
+		{
+			TGuardValue<bool> GuardIsRestoringState(bIsRestoringState, true);
+
+			const FTypedElementSelectionOptions SelectionOptions = FTypedElementSelectionOptions()
+				.SetAllowHidden(true)
+				.SetAllowGroups(false)
+				.SetAllowLegacyNotifications(false)
+				.SetWarnIfLocked(false);
+
+			// TODO: Work out the intersection of the before and after state instead of clearing and reselecting?
+			SetSelection(SelectedElements, SelectionOptions);
+		}
 	}
 }
 
@@ -273,7 +300,7 @@ void UTypedElementSelectionSet::OnElementListPreChange(const UTypedElementList*
 	check(InElementList == ElementList);
 	OnPreChangeDelegate.Broadcast(this);
 
-	if (!bIsRestoringFromTransaction)
+	if (!bIsRestoringState)
 	{
 		// Track the pre-change state for undo/redo
 		Modify();
diff --git a/Engine/Source/Runtime/TypedElementRuntime/Public/Elements/Framework/TypedElementSelectionSet.h b/Engine/Source/Runtime/TypedElementRuntime/Public/Elements/Framework/TypedElementSelectionSet.h
--- a/Engine/Source/Runtime/TypedElementRuntime/Public/Elements/Framework/TypedElementSelectionSet.h
+++ b/Engine/Source/Runtime/TypedElementRuntime/Public/Elements/Framework/TypedElementSelectionSet.h
@@ -82,12 +82,41 @@ struct FTypedElementSelectionSetState
 
 	friend class UTypedElementSelectionSet;
 
-private:
-	UPROPERTY()
-	TArray<uint8> StoredSelectionSetData;
+	FTypedElementSelectionSetState() = default;
+
+	FTypedElementSelectionSetState(FTypedElementSelectionSetState&&) = default;
+	FTypedElementSelectionSetState& operator=(FTypedElementSelectionSetState&&) = default;
+
+	FTypedElementSelectionSetState(const FTypedElementSelectionSetState& InOther)
+		: CreatedFromSelectionSet(InOther.CreatedFromSelectionSet)
+	{
+		TransactedElements.Reserve(InOther.TransactedElements.Num());
+		for (const TUniquePtr<ITypedElementTransactedElement>& OtherTransactedElement : InOther.TransactedElements)
+		{
+			TransactedElements.Emplace(OtherTransactedElement->Clone());
+		}
+	}
 
+	FTypedElementSelectionSetState& operator=(const FTypedElementSelectionSetState& InOther)
+	{
+		if (&InOther != this)
+		{
+			CreatedFromSelectionSet = InOther.CreatedFromSelectionSet;
+
+			TransactedElements.Reset(InOther.TransactedElements.Num());
+			for (const TUniquePtr<ITypedElementTransactedElement>& OtherTransactedElement : InOther.TransactedElements)
+			{
+				TransactedElements.Emplace(OtherTransactedElement->Clone());
+			}
+		}
+		return *this;
+	}
+
+private:
 	UPROPERTY()
 	TWeakObjectPtr<const UTypedElementSelectionSet> CreatedFromSelectionSet;
+
+	TArray<TUniquePtr<ITypedElementTransactedElement>> TransactedElements;
 };
 
 /**
@@ -104,6 +133,8 @@ public:
 
 	//~ UObject interface
 #if WITH_EDITOR
+	virtual void PreEditUndo() override;
+	virtual void PostEditUndo() override;
 	virtual bool Modify(bool bAlwaysMarkDirty = true) override;
 #endif	// WITH_EDITOR
 	virtual void Serialize(FArchive& Ar) override;
@@ -479,6 +510,9 @@ private:
 	/** Delegate that is invoked whenever the underlying element list has been changed. */
 	FOnChanged OnChangedDelegate;
 
-	/** Set when we are currently restoring the selection state from a transaction (undo/redo) */
-	bool bIsRestoringFromTransaction = false;
+	/** Set when we are currently restoring the selection state (eg, from an undo/redo) */
+	bool bIsRestoringState = false;
+
+	/** Set between a PreEditUndo->PostEditUndo call, to hold the state loaded during Serialize and applied in PostEditUndo */
+	TUniquePtr<FTypedElementSelectionSetState> PendingUndoRedoState;
 };
diff --git a/Engine/Source/Runtime/TypedElementRuntime/Public/Elements/Interfaces/TypedElementSelectionInterface.h b/Engine/Source/Runtime/TypedElementRuntime/Public/Elements/Interfaces/TypedElementSelectionInterface.h
--- a/Engine/Source/Runtime/TypedElementRuntime/Public/Elements/Interfaces/TypedElementSelectionInterface.h
+++ b/Engine/Source/Runtime/TypedElementRuntime/Public/Elements/Interfaces/TypedElementSelectionInterface.h
@@ -62,6 +62,55 @@ private:
 	bool bWarnIfLocked = false;
 };
 
+class ITypedElementTransactedElement
+{
+public:
+	virtual ~ITypedElementTransactedElement() = default;
+
+	TUniquePtr<ITypedElementTransactedElement> Clone() const
+	{
+		TUniquePtr<ITypedElementTransactedElement> Cloned = CloneImpl();
+		checkf(Cloned, TEXT("ITypedElementTransactedElement derived types must implement a valid Clone function!"));
+		return Cloned;
+	}
+
+	FTypedElementHandle GetElement() const
+	{
+		return GetElementImpl();
+	}
+
+	FTypedHandleTypeId GetElementType() const
+	{
+		return TypeId;
+	}
+
+	void SetElement(const FTypedElementHandle& InElementHandle)
+	{
+		SetElementType(InElementHandle.GetId().GetTypeId());
+		SetElementImpl(InElementHandle);
+	}
+
+	void SetElementType(const FTypedHandleTypeId InTypeId)
+	{
+		TypeId = InTypeId;
+	}
+
+	void Serialize(FArchive& InArchive)
+	{
+		checkf(!InArchive.IsPersistent(), TEXT("ITypedElementTransactedElement can only be serialized by transient archives!"));
+		SerializeImpl(InArchive);
+	}
+
+protected:
+	virtual TUniquePtr<ITypedElementTransactedElement> CloneImpl() const = 0;
+	virtual FTypedElementHandle GetElementImpl() const = 0;
+	virtual void SetElementImpl(const FTypedElementHandle& InElementHandle) = 0;
+	virtual void SerializeImpl(FArchive& InArchive) = 0;
+
+private:
+	FTypedHandleTypeId TypeId = 0;
+};
+
 UCLASS(Abstract)
 class TYPEDELEMENTRUNTIME_API UTypedElementSelectionInterface : public UTypedElementInterface
 {
@@ -118,16 +167,37 @@ public:
 	virtual bool ShouldPreventTransactions(const FTypedElementHandle& InElementHandle) { return false; }
 
 	/**
-	 * Write the information needed to find the given element again when this transaction is replayed for undo/redo.
-	 * @note The data written must match the data read by ReadTransactedElement.
+	 * Create a transacted element instance that can be used to save the given element for undo/redo.
 	 */
-	virtual void WriteTransactedElement(const FTypedElementHandle& InElementHandle, FArchive& InArchive) {}
+	TUniquePtr<ITypedElementTransactedElement> CreateTransactedElement(const FTypedElementHandle& InElementHandle)
+	{
+		TUniquePtr<ITypedElementTransactedElement> TransactedElement = CreateTransactedElementImpl();
+		if (TransactedElement)
+		{
+			TransactedElement->SetElement(InElementHandle);
+		}
+		return TransactedElement;
+	}
 
 	/**
-	 * Read the information needed to be find a previously tracked element, and attempt to get it now that this transaction is being replayed for undo/redo.
-	 * @note The data read must match the data written by WriteTransactedElement.
+	 * Create a transacted element instance that can be used to load an element previously saved for undo/redo.
+	 */
+	TUniquePtr<ITypedElementTransactedElement> CreateTransactedElement(const FTypedHandleTypeId InTypeId)
+	{
+		TUniquePtr<ITypedElementTransactedElement> TransactedElement = CreateTransactedElementImpl();
+		if (TransactedElement)
+		{
+			TransactedElement->SetElementType(InTypeId);
+		}
+		return TransactedElement;
+	}
+
+protected:
+	/**
+	 * Create a transacted element instance that can be used to save/load elements of the implementation type for undo/redo.
+	 * @note The instance returned from this function must have either SetElement or SetElementType called on it prior to being used.
 	 */
-	virtual FTypedElementHandle ReadTransactedElement(FArchive& InArchive) { return FTypedElementHandle(); }
+	virtual TUniquePtr<ITypedElementTransactedElement> CreateTransactedElementImpl() { return nullptr; }
 };
 
 template <>
@@ -141,6 +211,5 @@ struct TTypedElement<UTypedElementSelectionInterface> : public TTypedElementBase
 	bool AllowSelectionModifiers(const UTypedElementList* InSelectionSet) const { return InterfacePtr->AllowSelectionModifiers(*this, InSelectionSet); }
 	FTypedElementHandle GetSelectionElement(const UTypedElementList* InCurrentSelection, const ETypedElementSelectionMethod InSelectionMethod) const { return InterfacePtr->GetSelectionElement(*this, InCurrentSelection, InSelectionMethod); }
 	bool ShouldPreventTransactions() const { return InterfacePtr->ShouldPreventTransactions(*this); }
-	void WriteTransactedElement(FArchive& InArchive) const { return InterfacePtr->WriteTransactedElement(*this, InArchive); }
-	FTypedElementHandle ReadTransactedElement(FArchive& InArchive) const { return InterfacePtr->ReadTransactedElement(InArchive); }
+	TUniquePtr<ITypedElementTransactedElement> CreateTransactedElement() const { return InterfacePtr->CreateTransactedElement(*this); }
 };
