Cache the selected element state in PreEditUndo and immediately clear the active selection

Follow-up to 15394195, as the elements may no longer be valid to remove if we wait until PostEditUndo to clear the selection set

diff --git a/Engine/Source/Runtime/TypedElementRuntime/Private/Elements/Framework/TypedElementSelectionSet.cpp b/Engine/Source/Runtime/TypedElementRuntime/Private/Elements/Framework/TypedElementSelectionSet.cpp
--- a/Engine/Source/Runtime/TypedElementRuntime/Private/Elements/Framework/TypedElementSelectionSet.cpp
+++ b/Engine/Source/Runtime/TypedElementRuntime/Private/Elements/Framework/TypedElementSelectionSet.cpp
@@ -22,7 +22,11 @@ void UTypedElementSelectionSet::PreEditUndo()
 	Super::PreEditUndo();
 
 	checkf(!PendingUndoRedoState, TEXT("PendingUndoRedoState was set! Missing call to PostEditUndo?"));
-	PendingUndoRedoState = MakeUnique<FTypedElementSelectionSetState>();
+	PendingUndoRedoState = MakeUnique<FTypedElementSelectionSetState>(GetCurrentSelectionState());
+
+	FTypedElementSelectionSetState EmptySelectionState;
+	EmptySelectionState.CreatedFromSelectionSet = this;
+	RestoreSelectionState(EmptySelectionState);
 }
 
 void UTypedElementSelectionSet::PostEditUndo()
@@ -61,9 +65,19 @@ void UTypedElementSelectionSet::Serialize(FArchive& Ar)
 {
 	checkf(!Ar.IsPersistent(), TEXT("UTypedElementSelectionSet can only be serialized by transient archives!"));
 
+	const bool bIsUndoRedo = PendingUndoRedoState && Ar.IsTransacting();
+
+	FTypedElementSelectionSetState TmpSelectionState;
+	TmpSelectionState.CreatedFromSelectionSet = this;
+
+	FTypedElementSelectionSetState& SelectionState = bIsUndoRedo ? *PendingUndoRedoState : TmpSelectionState;
+
 	if (Ar.IsSaving())
 	{
-		FTypedElementSelectionSetState SelectionState = ElementList ? GetCurrentSelectionState() : FTypedElementSelectionSetState();
+		if (ElementList && !bIsUndoRedo)
+		{
+			TmpSelectionState = GetCurrentSelectionState();
+		}
 
 		int32 NumTransactedElements = SelectionState.TransactedElements.Num();
 		Ar << NumTransactedElements;
@@ -80,16 +94,10 @@ void UTypedElementSelectionSet::Serialize(FArchive& Ar)
 	{
 		UTypedElementRegistry* Registry = UTypedElementRegistry::GetInstance();
 
-		const bool bIsUndoRedo = PendingUndoRedoState && Ar.IsTransacting();
-		FTypedElementSelectionSetState TmpSelectionState;
-
-		FTypedElementSelectionSetState& SelectionState = bIsUndoRedo ? *PendingUndoRedoState : TmpSelectionState;
-		SelectionState.CreatedFromSelectionSet = this;
-
 		int32 NumTransactedElements = 0;
 		Ar << NumTransactedElements;
 
-		SelectionState.TransactedElements.Reserve(NumTransactedElements);
+		SelectionState.TransactedElements.Reset(NumTransactedElements);
 		for (int32 TransactedElementIndex = 0; TransactedElementIndex < NumTransactedElements; ++TransactedElementIndex)
 		{
 			FTypedHandleTypeId TypeId = 0;
diff --git a/Engine/Source/Runtime/TypedElementRuntime/Public/Elements/Framework/TypedElementSelectionSet.h b/Engine/Source/Runtime/TypedElementRuntime/Public/Elements/Framework/TypedElementSelectionSet.h
--- a/Engine/Source/Runtime/TypedElementRuntime/Public/Elements/Framework/TypedElementSelectionSet.h
+++ b/Engine/Source/Runtime/TypedElementRuntime/Public/Elements/Framework/TypedElementSelectionSet.h
@@ -513,6 +513,10 @@ private:
 	/** Set when we are currently restoring the selection state (eg, from an undo/redo) */
 	bool bIsRestoringState = false;
 
-	/** Set between a PreEditUndo->PostEditUndo call, to hold the state loaded during Serialize and applied in PostEditUndo */
+	/**
+	 * Set between a PreEditUndo->PostEditUndo call.
+	 * Initially holds the state cached during PreEditUndo which is saved during Serialize, 
+	 * and is then replaced with the state loaded from Serialize to be applied in PostEditUndo.
+	 */
 	TUniquePtr<FTypedElementSelectionSetState> PendingUndoRedoState;
 };
