Fixed selection annotation getting out-of-sync during object re-instancing

The selection set code that handled updating the selection state for an object re-instance was only updating the underlying list, and not invoking any of the higher-level selection customization logic. This resulted in selection annotations (updated by the level editor selection customizations) that were out-of-date, meaning that things like AActor::IsSelected() could return the incorrect result.

diff --git a/Engine/Source/Runtime/TypedElementRuntime/Private/Elements/Framework/TypedElementSelectionSet.cpp b/Engine/Source/Runtime/TypedElementRuntime/Private/Elements/Framework/TypedElementSelectionSet.cpp
--- a/Engine/Source/Runtime/TypedElementRuntime/Private/Elements/Framework/TypedElementSelectionSet.cpp
+++ b/Engine/Source/Runtime/TypedElementRuntime/Private/Elements/Framework/TypedElementSelectionSet.cpp
@@ -322,18 +322,34 @@ void UTypedElementSelectionSet::OnElementReplaced(TArrayView<const TTuple<FTyped
 			break;
 		}
 
-		if (ElementList->Remove(ReplacedElement.Key))
+		bool bDeselectedElement = false;
 		{
-			if (ReplacedElement.Value)
+			FTypedElementSelectionSetElement SelectionSetElementToRemove = ResolveSelectionSetElement(ReplacedElement.Key);
+			if (SelectionSetElementToRemove && SelectionSetElementToRemove.CanDeselectElement(SelectionOptions))
 			{
-				SelectElement(ReplacedElement.Value, SelectionOptions);
+				bDeselectedElement = SelectionSetElementToRemove.DeselectElement(SelectionOptions);
 			}
+			else
+			{
+				bDeselectedElement = ElementList->Remove(ReplacedElement.Key);
+			}
+		}
+
+		if (bDeselectedElement && ReplacedElement.Value)
+		{
+			SelectElement(ReplacedElement.Value, SelectionOptions);
 		}
 	}
 }
 
 void UTypedElementSelectionSet::OnElementUpdated(TArrayView<const FTypedElementHandle> InUpdatedElements)
 {
+	const FTypedElementSelectionOptions SelectionOptions = FTypedElementSelectionOptions()
+		.SetAllowHidden(true)
+		.SetAllowGroups(false)
+		.SetAllowLegacyNotifications(false)
+		.SetWarnIfLocked(false);
+
 	if (ElementList->Num() == 0)
 	{
 		return;
@@ -341,9 +357,21 @@ void UTypedElementSelectionSet::OnElementUpdated(TArrayView<const FTypedElementH
 	
 	for (const FTypedElementHandle& UpdatedElement : InUpdatedElements)
 	{
-		if (ElementList->Remove(UpdatedElement))
+		FTypedElementSelectionSetElement SelectionSetElementToUpdate = ResolveSelectionSetElement(UpdatedElement);
+
+		bool bDeselectedElement = false;
+		if (SelectionSetElementToUpdate && SelectionSetElementToUpdate.CanDeselectElement(SelectionOptions))
+		{
+			bDeselectedElement = SelectionSetElementToUpdate.DeselectElement(SelectionOptions);
+		}
+		else
+		{
+			bDeselectedElement = ElementList->Remove(UpdatedElement);
+		}
+
+		if (bDeselectedElement && SelectionSetElementToUpdate.CanSelectElement(SelectionOptions))
 		{
-			ElementList->Add(UpdatedElement);
+			SelectionSetElementToUpdate.SelectElement(SelectionOptions);
 		}
 	}
 }
