Always emit an OnPreChange notification for an element list

This allows things that cache based on the list state to consistently use OnPreChange to dirty their cache, as previously it was possible for an OnPreChange to dirty a cache that was rebuilt while the list was still changing (at which point further OnPreChange events were suppressed).

It would also be possible for the caller to bind to OnChanged alongside OnPreChange, and dirty their cache from both events, however OnChanged can be suppressed if ClearPendingChanges is called, which makes it unreliable for cache invalidation.

diff --git a/Engine/Source/Editor/UnrealEd/Private/EditorServer.cpp b/Engine/Source/Editor/UnrealEd/Private/EditorServer.cpp
--- a/Engine/Source/Editor/UnrealEd/Private/EditorServer.cpp
+++ b/Engine/Source/Editor/UnrealEd/Private/EditorServer.cpp
@@ -2168,8 +2168,10 @@ void UEditorEngine::EditorDestroyWorld( FWorldContext & Context, const FText& Cl
 		NewWorld->RemoveFromRoot();
 	}
 
-	CheckForWorldGCLeaks( NewWorld, WorldPackage );
+	// Verify that any pending typed elements have been cleanly destroyed from the GC above
+	UTypedElementRegistry::GetInstance()->ProcessDeferredElementsToDestroy();
 
+	CheckForWorldGCLeaks( NewWorld, WorldPackage );
 }
 
 bool UEditorEngine::ShouldAbortBecauseOfPIEWorld()
diff --git a/Engine/Source/Runtime/TypedElementFramework/Private/Elements/Framework/TypedElementList.cpp b/Engine/Source/Runtime/TypedElementFramework/Private/Elements/Framework/TypedElementList.cpp
--- a/Engine/Source/Runtime/TypedElementFramework/Private/Elements/Framework/TypedElementList.cpp
+++ b/Engine/Source/Runtime/TypedElementFramework/Private/Elements/Framework/TypedElementList.cpp
@@ -375,10 +375,7 @@ void UTypedElementList::ClearPendingChanges()
 
 void UTypedElementList::NoteListMayChange()
 {
-	if (!bHasPendingNotify)
-	{
-		OnPreChangeDelegate.Broadcast(this);
-	}
+	OnPreChangeDelegate.Broadcast(this);
 }
 
 void UTypedElementList::NoteListChanged(const EChangeType InChangeType, const FTypedElementHandle& InElementHandle)
diff --git a/Engine/Source/Runtime/TypedElementFramework/Public/Elements/Framework/TypedElementList.h b/Engine/Source/Runtime/TypedElementFramework/Public/Elements/Framework/TypedElementList.h
--- a/Engine/Source/Runtime/TypedElementFramework/Public/Elements/Framework/TypedElementList.h
+++ b/Engine/Source/Runtime/TypedElementFramework/Public/Elements/Framework/TypedElementList.h
@@ -558,7 +558,7 @@ public:
 
 	/**
 	 * Access the delegate that is invoked whenever this element list is potentially about to change.
-	 * @note This may be called even if no actual change happens, though once a change does happen it won't be called again until after the next call to NotifyPendingChanges.
+	 * @note This may be called even if no actual change happens, so may be called multiple times without a corresponding OnChanged notification.
 	 */
 	DECLARE_EVENT_OneParam(UTypedElementList, FOnPreChange, const UTypedElementList* /*InElementList*/);
 	FOnPreChange& OnPreChange()
diff --git a/Engine/Source/Runtime/TypedElementRuntime/Public/Elements/Framework/TypedElementSelectionSet.h b/Engine/Source/Runtime/TypedElementRuntime/Public/Elements/Framework/TypedElementSelectionSet.h
--- a/Engine/Source/Runtime/TypedElementRuntime/Public/Elements/Framework/TypedElementSelectionSet.h
+++ b/Engine/Source/Runtime/TypedElementRuntime/Public/Elements/Framework/TypedElementSelectionSet.h
@@ -445,8 +445,8 @@ public:
 	}
 
 	/**
-	 * Access the delegate that is invoked whenever this element list is potentially about to change.
-	 * @note This may be called even if no actual change happens, though once a change does happen it won't be called again until after the next call to NotifyPendingChanges.
+	 * Access the delegate that is invoked whenever the underlying element list is potentially about to change.
+	 * @note This may be called even if no actual change happens, so may be called multiple times without a corresponding OnChanged notification.
 	 */
 	DECLARE_EVENT_OneParam(UTypedElementSelectionSet, FOnPreChange, const UTypedElementSelectionSet* /*InElementSelectionSet*/);
 	FOnPreChange& OnPreChange()
