Converted UTypedElementList to no longer be a UObject

Storing typed elements in UObjects can easily lead to reference leaks if the elements aren't cleared prior to the UObject being left pending GC. This actually made UTypedElementList tricky to use, as you had to remember to manually empty it when you'd finished with it to avoid reference leaks, and we've had several cases now where that was done incorrectly.

To address this issue, we've moved element lists back to being a normal C++ class, FTypedElementList. However, unlike the original version of FTypedElementList (which was itself a UStruct), this version is always heap-allocated and referenced via a TSharedPtr/TSharedRef.

This gives us a nice middle-ground of a well defined lifetime (ie, no lingering references prior to GC) while still being efficient to pass around, including for scripting APIs via FTypedElementListProxy (which just wraps the TSharedPtr in a UStruct).

The downside of this approach is that we need to wrap the FTypedElementList functions that we want to expose to the scripting API (see UTypedElementListLibrary), however that is a far more reasonable burden than requring every user of the typed element framework to know and understand that UTypedElementList had to be manually cleared to avoid potentially hard to find reference leaks (especially if via leaked via scripting APIs).

The core of this change is to TypedElementList.h/.cpp, with TypedElementListFwd.h existing to forward declare the pointer types, and TypedElementListProxy.h and TypedElementListLibrary.h existing to declare the proxy type and wrapped functions used for scripting APIs. TypedElementSelectionInterface.h (and its implementations) provide an example of using FTypedElementListProxy within a scripting API, and the rest of the change is mostly just fallout to transform const UTypedElementList* to FTypedElementListConstRef and UTypedElementList* to FTypedElementListRef.

diff --git a/Engine/Source/Editor/AssetPlacementEdMode/Private/Tools/PlacementEraseTool.cpp b/Engine/Source/Editor/AssetPlacementEdMode/Private/Tools/PlacementEraseTool.cpp
--- a/Engine/Source/Editor/AssetPlacementEdMode/Private/Tools/PlacementEraseTool.cpp
+++ b/Engine/Source/Editor/AssetPlacementEdMode/Private/Tools/PlacementEraseTool.cpp
@@ -66,7 +66,7 @@ void UPlacementModeEraseTool::OnTick(float DeltaTime)
 		return;
 	}
 
-	TGCObjectScopeGuard<UTypedElementList> ElementsToDelete(UTypedElementRegistry::GetInstance()->CreateElementList());
+	FTypedElementListRef ElementsToDelete = UTypedElementRegistry::GetInstance()->CreateElementList();
 
 	TArray<FTypedElementHandle> HitElements = GetElementsInBrushRadius(LastDeviceInputRay);
 	for (const FTypedElementHandle& HitElement : HitElements)
@@ -91,16 +91,14 @@ void UPlacementModeEraseTool::OnTick(float DeltaTime)
 			}
 			else
 			{
-				ElementsToDelete.Get()->Add(HitElement);
+				ElementsToDelete->Add(HitElement);
 			}
 		}
 	}
 
-	if (ElementsToDelete.Get()->HasElements())
+	if (ElementsToDelete->HasElements())
 	{
-		TGCObjectScopeGuard<UTypedElementList> NormalizedElementsToDelete(SelectionSet->GetNormalizedElementList(ElementsToDelete.Get(), FTypedElementSelectionNormalizationOptions()));
-		ElementCommonActions->DeleteNormalizedElements(NormalizedElementsToDelete.Get(), AssetEditorContext->GetEditingWorld(), SelectionSet, FTypedElementDeletionOptions());
-		NormalizedElementsToDelete.Get()->Empty();
-		ElementsToDelete.Get()->Empty();
+		FTypedElementListRef NormalizedElementsToDelete = SelectionSet->GetNormalizedElementList(ElementsToDelete, FTypedElementSelectionNormalizationOptions());
+		ElementCommonActions->DeleteNormalizedElements(NormalizedElementsToDelete, AssetEditorContext->GetEditingWorld(), SelectionSet, FTypedElementDeletionOptions());
 	}
 }
diff --git a/Engine/Source/Editor/EditorFramework/Private/Elements/Framework/TypedElementViewportInteraction.cpp b/Engine/Source/Editor/EditorFramework/Private/Elements/Framework/TypedElementViewportInteraction.cpp
--- a/Engine/Source/Editor/EditorFramework/Private/Elements/Framework/TypedElementViewportInteraction.cpp
+++ b/Engine/Source/Editor/EditorFramework/Private/Elements/Framework/TypedElementViewportInteraction.cpp
@@ -131,7 +131,7 @@ void FTypedElementViewportInteractionCustomization::MirrorElement(const TTypedEl
 }
 
 
-void UTypedElementViewportInteraction::BeginGizmoManipulation(const UTypedElementList* InElementsToMove, const UE::Widget::EWidgetMode InWidgetMode)
+void UTypedElementViewportInteraction::BeginGizmoManipulation(FTypedElementListConstRef InElementsToMove, const UE::Widget::EWidgetMode InWidgetMode)
 {
 	{
 		TMap<FTypedHandleTypeId, TArray<FTypedElementHandle>> ElementsToMoveByType;
@@ -156,7 +156,7 @@ void UTypedElementViewportInteraction::BeginGizmoManipulation(const UTypedElemen
 	});
 }
 
-void UTypedElementViewportInteraction::UpdateGizmoManipulation(const UTypedElementList* InElementsToMove, const UE::Widget::EWidgetMode InWidgetMode, const EAxisList::Type InDragAxis, const FInputDeviceState& InInputState, const FTransform& InDeltaTransform)
+void UTypedElementViewportInteraction::UpdateGizmoManipulation(FTypedElementListConstRef InElementsToMove, const UE::Widget::EWidgetMode InWidgetMode, const EAxisList::Type InDragAxis, const FInputDeviceState& InInputState, const FTransform& InDeltaTransform)
 {
 	InElementsToMove->ForEachElementHandle([this, InWidgetMode, InDragAxis, &InInputState, &InDeltaTransform](const FTypedElementHandle& InElementToMove)
 	{
@@ -171,7 +171,7 @@ void UTypedElementViewportInteraction::UpdateGizmoManipulation(const UTypedEleme
 	});
 }
 
-void UTypedElementViewportInteraction::EndGizmoManipulation(const UTypedElementList* InElementsToMove, const UE::Widget::EWidgetMode InWidgetMode, const ETypedElementViewportInteractionGizmoManipulationType InManipulationType)
+void UTypedElementViewportInteraction::EndGizmoManipulation(FTypedElementListConstRef InElementsToMove, const UE::Widget::EWidgetMode InWidgetMode, const ETypedElementViewportInteractionGizmoManipulationType InManipulationType)
 {
 	InElementsToMove->ForEachElementHandle([this, InWidgetMode, InManipulationType](const FTypedElementHandle& InElementToMove)
 	{
diff --git a/Engine/Source/Editor/EditorFramework/Public/Elements/Framework/TypedElementViewportInteraction.h b/Engine/Source/Editor/EditorFramework/Public/Elements/Framework/TypedElementViewportInteraction.h
--- a/Engine/Source/Editor/EditorFramework/Public/Elements/Framework/TypedElementViewportInteraction.h
+++ b/Engine/Source/Editor/EditorFramework/Public/Elements/Framework/TypedElementViewportInteraction.h
@@ -4,11 +4,11 @@
 
 #include "InputState.h"
 #include "UnrealWidgetFwd.h"
+#include "Elements/Framework/TypedElementListFwd.h"
 #include "Elements/Framework/TypedElementInterfaceCustomization.h"
 #include "Elements/Interfaces/TypedElementWorldInterface.h"
 #include "TypedElementViewportInteraction.generated.h"
 
-class UTypedElementList;
 class UTypedElementSelectionSet;
 
 enum class ETypedElementViewportInteractionGizmoManipulationType : uint8
@@ -93,17 +93,17 @@ public:
 	/**
 	 * Notify that the gizmo is potentially about to start manipulating the transform of the given set of elements (calculated from calling GetSelectedElementsToMove).
 	 */
-	void BeginGizmoManipulation(const UTypedElementList* InElementsToMove, const UE::Widget::EWidgetMode InWidgetMode);
+	void BeginGizmoManipulation(FTypedElementListConstRef InElementsToMove, const UE::Widget::EWidgetMode InWidgetMode);
 
 	/**
 	 * Notify that the gizmo has manipulated the transform of the given set of elements (calculated from calling GetSelectedElementsToMove) by the given delta.
 	 */
-	void UpdateGizmoManipulation(const UTypedElementList* InElementsToMove, const UE::Widget::EWidgetMode InWidgetMode, const EAxisList::Type InDragAxis, const FInputDeviceState& InInputState, const FTransform& InDeltaTransform);
+	void UpdateGizmoManipulation(FTypedElementListConstRef InElementsToMove, const UE::Widget::EWidgetMode InWidgetMode, const EAxisList::Type InDragAxis, const FInputDeviceState& InInputState, const FTransform& InDeltaTransform);
 	
 	/**
 	 * Notify that the gizmo has finished manipulating the transform of the given set of elements (calculated from calling GetSelectedElementsToMove).
 	 */
-	void EndGizmoManipulation(const UTypedElementList* InElementsToMove, const UE::Widget::EWidgetMode InWidgetMode, const ETypedElementViewportInteractionGizmoManipulationType InManipulationType);
+	void EndGizmoManipulation(FTypedElementListConstRef InElementsToMove, const UE::Widget::EWidgetMode InWidgetMode, const ETypedElementViewportInteractionGizmoManipulationType InManipulationType);
 
 	/**
 	 * Apply the given delta to the specified element.
diff --git a/Engine/Source/Editor/LevelEditor/Private/Elements/Actor/ActorElementLevelEditorSelectionCustomization.cpp b/Engine/Source/Editor/LevelEditor/Private/Elements/Actor/ActorElementLevelEditorSelectionCustomization.cpp
--- a/Engine/Source/Editor/LevelEditor/Private/Elements/Actor/ActorElementLevelEditorSelectionCustomization.cpp
+++ b/Engine/Source/Editor/LevelEditor/Private/Elements/Actor/ActorElementLevelEditorSelectionCustomization.cpp
@@ -30,23 +30,23 @@ bool FActorElementLevelEditorSelectionCustomization::CanDeselectElement(const TT
 	return CanDeselectActorElement(InElementSelectionHandle, InSelectionOptions);
 }
 
-bool FActorElementLevelEditorSelectionCustomization::SelectElement(const TTypedElement<UTypedElementSelectionInterface>& InElementSelectionHandle, UTypedElementList* InSelectionSet, const FTypedElementSelectionOptions& InSelectionOptions)
+bool FActorElementLevelEditorSelectionCustomization::SelectElement(const TTypedElement<UTypedElementSelectionInterface>& InElementSelectionHandle, FTypedElementListRef InSelectionSet, const FTypedElementSelectionOptions& InSelectionOptions)
 {
 	return SelectActorElement(InElementSelectionHandle, InSelectionSet, InSelectionOptions);
 }
 
-bool FActorElementLevelEditorSelectionCustomization::DeselectElement(const TTypedElement<UTypedElementSelectionInterface>& InElementSelectionHandle, UTypedElementList* InSelectionSet, const FTypedElementSelectionOptions& InSelectionOptions)
+bool FActorElementLevelEditorSelectionCustomization::DeselectElement(const TTypedElement<UTypedElementSelectionInterface>& InElementSelectionHandle, FTypedElementListRef InSelectionSet, const FTypedElementSelectionOptions& InSelectionOptions)
 {
 	return DeselectActorElement(InElementSelectionHandle, InSelectionSet, InSelectionOptions);
 }
 
-bool FActorElementLevelEditorSelectionCustomization::AllowSelectionModifiers(const TTypedElement<UTypedElementSelectionInterface>& InElementSelectionHandle, const UTypedElementList* InSelectionSet)
+bool FActorElementLevelEditorSelectionCustomization::AllowSelectionModifiers(const TTypedElement<UTypedElementSelectionInterface>& InElementSelectionHandle, FTypedElementListConstRef InSelectionSet)
 {
 	// Ctrl or Shift clicking an actor is the same as regular clicking when components are selected
 	return !InSelectionSet->HasElementsOfType(NAME_Components);
 }
 
-FTypedElementHandle FActorElementLevelEditorSelectionCustomization::GetSelectionElement(const TTypedElement<UTypedElementSelectionInterface>& InElementSelectionHandle, const UTypedElementList* InCurrentSelection, const ETypedElementSelectionMethod InSelectionMethod)
+FTypedElementHandle FActorElementLevelEditorSelectionCustomization::GetSelectionElement(const TTypedElement<UTypedElementSelectionInterface>& InElementSelectionHandle, FTypedElementListConstRef InCurrentSelection, const ETypedElementSelectionMethod InSelectionMethod)
 {
 	if (AActor* ConsideredActor = ActorElementDataUtil::GetActorFromHandle(InElementSelectionHandle))
 	{
@@ -59,7 +59,7 @@ FTypedElementHandle FActorElementLevelEditorSelectionCustomization::GetSelection
 	return InElementSelectionHandle;
 }
 
-void FActorElementLevelEditorSelectionCustomization::GetNormalizedElements(const TTypedElement<UTypedElementSelectionInterface>& InElementSelectionHandle, const UTypedElementList* InSelectionSet, const FTypedElementSelectionNormalizationOptions& InNormalizationOptions, UTypedElementList* OutNormalizedElements)
+void FActorElementLevelEditorSelectionCustomization::GetNormalizedElements(const TTypedElement<UTypedElementSelectionInterface>& InElementSelectionHandle, FTypedElementListConstRef InSelectionSet, const FTypedElementSelectionNormalizationOptions& InNormalizationOptions, FTypedElementListRef OutNormalizedElements)
 {
 	AActor* Actor = ActorElementDataUtil::GetActorFromHandleChecked(InElementSelectionHandle);
 
@@ -143,7 +143,7 @@ bool FActorElementLevelEditorSelectionCustomization::CanDeselectActorElement(con
 	return true;
 }
 
-bool FActorElementLevelEditorSelectionCustomization::SelectActorElement(const TTypedElement<UTypedElementSelectionInterface>& InActorSelectionHandle, UTypedElementList* InSelectionSet, const FTypedElementSelectionOptions& InSelectionOptions)
+bool FActorElementLevelEditorSelectionCustomization::SelectActorElement(const TTypedElement<UTypedElementSelectionInterface>& InActorSelectionHandle, FTypedElementListRef InSelectionSet, const FTypedElementSelectionOptions& InSelectionOptions)
 {
 	AActor* Actor = ActorElementDataUtil::GetActorFromHandleChecked(InActorSelectionHandle);
 
@@ -213,7 +213,7 @@ bool FActorElementLevelEditorSelectionCustomization::SelectActorElement(const TT
 	return true;
 }
 
-bool FActorElementLevelEditorSelectionCustomization::DeselectActorElement(const TTypedElement<UTypedElementSelectionInterface>& InActorSelectionHandle, UTypedElementList* InSelectionSet, const FTypedElementSelectionOptions& InSelectionOptions)
+bool FActorElementLevelEditorSelectionCustomization::DeselectActorElement(const TTypedElement<UTypedElementSelectionInterface>& InActorSelectionHandle, FTypedElementListRef InSelectionSet, const FTypedElementSelectionOptions& InSelectionOptions)
 {
 	AActor* Actor = ActorElementDataUtil::GetActorFromHandleChecked(InActorSelectionHandle);
 
@@ -256,7 +256,7 @@ bool FActorElementLevelEditorSelectionCustomization::DeselectActorElement(const
 	
 	// Deselect and unbind the override delegates for the components on the selected actor
 	{
-		FTypedElementListLegacySyncScopedBatch LegacySyncBatch(InSelectionSet, InSelectionOptions.AllowLegacyNotifications());
+		FTypedElementListLegacySyncScopedBatch LegacySyncBatch(*InSelectionSet, InSelectionOptions.AllowLegacyNotifications());
 
 		for (UActorComponent* Component : Actor->GetComponents())
 		{
@@ -284,7 +284,7 @@ bool FActorElementLevelEditorSelectionCustomization::DeselectActorElement(const
 	return true;
 }
 
-bool FActorElementLevelEditorSelectionCustomization::SelectActorGroup(AGroupActor* InGroupActor, UTypedElementList* InSelectionSet, const FTypedElementSelectionOptions& InSelectionOptions, const bool bForce)
+bool FActorElementLevelEditorSelectionCustomization::SelectActorGroup(AGroupActor* InGroupActor, FTypedElementListRef InSelectionSet, const FTypedElementSelectionOptions& InSelectionOptions, const bool bForce)
 {
 	bool bSelectionChanged = false;
 
@@ -297,7 +297,7 @@ bool FActorElementLevelEditorSelectionCustomization::SelectActorGroup(AGroupActo
 	// Skip if the group is already selected, since this logic will have already run
 	if ((bForce || InGroupActor->IsLocked()) && !GroupSelectionHandle.IsElementSelected(InSelectionSet, FTypedElementIsSelectedOptions()) && CanSelectActorElement(GroupSelectionHandle, GroupSelectionOptions))
 	{
-		FTypedElementListLegacySyncScopedBatch LegacySyncBatch(InSelectionSet, InSelectionOptions.AllowLegacyNotifications());
+		FTypedElementListLegacySyncScopedBatch LegacySyncBatch(*InSelectionSet, InSelectionOptions.AllowLegacyNotifications());
 
 		bSelectionChanged |= SelectActorElement(GroupSelectionHandle, InSelectionSet, GroupSelectionOptions);
 
@@ -316,7 +316,7 @@ bool FActorElementLevelEditorSelectionCustomization::SelectActorGroup(AGroupActo
 	return bSelectionChanged;
 }
 
-bool FActorElementLevelEditorSelectionCustomization::DeselectActorGroup(AGroupActor* InGroupActor, UTypedElementList* InSelectionSet, const FTypedElementSelectionOptions& InSelectionOptions, const bool bForce)
+bool FActorElementLevelEditorSelectionCustomization::DeselectActorGroup(AGroupActor* InGroupActor, FTypedElementListRef InSelectionSet, const FTypedElementSelectionOptions& InSelectionOptions, const bool bForce)
 {
 	bool bSelectionChanged = false;
 
@@ -329,7 +329,7 @@ bool FActorElementLevelEditorSelectionCustomization::DeselectActorGroup(AGroupAc
 	// Skip if the group is already deselected, since this logic will have already run
 	if ((bForce || InGroupActor->IsLocked()) && GroupSelectionHandle.IsElementSelected(InSelectionSet, FTypedElementIsSelectedOptions()) && CanDeselectActorElement(GroupSelectionHandle, GroupSelectionOptions))
 	{
-		FTypedElementListLegacySyncScopedBatch LegacySyncBatch(InSelectionSet, InSelectionOptions.AllowLegacyNotifications());
+		FTypedElementListLegacySyncScopedBatch LegacySyncBatch(*InSelectionSet, InSelectionOptions.AllowLegacyNotifications());
 
 		bSelectionChanged |= DeselectActorElement(GroupSelectionHandle, InSelectionSet, GroupSelectionOptions);
 
@@ -348,7 +348,7 @@ bool FActorElementLevelEditorSelectionCustomization::DeselectActorGroup(AGroupAc
 	return bSelectionChanged;
 }
 
-void FActorElementLevelEditorSelectionCustomization::AppendNormalizedActors(AActor* InActor, const UTypedElementList* InSelectionSet, const FTypedElementSelectionNormalizationOptions& InNormalizationOptions, UTypedElementList* OutNormalizedElements)
+void FActorElementLevelEditorSelectionCustomization::AppendNormalizedActors(AActor* InActor, FTypedElementListConstRef InSelectionSet, const FTypedElementSelectionNormalizationOptions& InNormalizationOptions, FTypedElementListRef OutNormalizedElements)
 {
 	if (InNormalizationOptions.FollowAttachment())
 	{
diff --git a/Engine/Source/Editor/LevelEditor/Private/Elements/Actor/ActorElementLevelEditorSelectionCustomization.h b/Engine/Source/Editor/LevelEditor/Private/Elements/Actor/ActorElementLevelEditorSelectionCustomization.h
--- a/Engine/Source/Editor/LevelEditor/Private/Elements/Actor/ActorElementLevelEditorSelectionCustomization.h
+++ b/Engine/Source/Editor/LevelEditor/Private/Elements/Actor/ActorElementLevelEditorSelectionCustomization.h
@@ -12,20 +12,20 @@ class FActorElementLevelEditorSelectionCustomization : public FTypedElementSelec
 public:
 	virtual bool CanSelectElement(const TTypedElement<UTypedElementSelectionInterface>& InElementSelectionHandle, const FTypedElementSelectionOptions& InSelectionOptions) override;
 	virtual bool CanDeselectElement(const TTypedElement<UTypedElementSelectionInterface>& InElementSelectionHandle, const FTypedElementSelectionOptions& InSelectionOptions) override;
-	virtual bool SelectElement(const TTypedElement<UTypedElementSelectionInterface>& InElementSelectionHandle, UTypedElementList* InSelectionSet, const FTypedElementSelectionOptions& InSelectionOptions) override;
-	virtual bool DeselectElement(const TTypedElement<UTypedElementSelectionInterface>& InElementSelectionHandle, UTypedElementList* InSelectionSet, const FTypedElementSelectionOptions& InSelectionOptions) override;
-	virtual bool AllowSelectionModifiers(const TTypedElement<UTypedElementSelectionInterface>& InElementSelectionHandle, const UTypedElementList* InSelectionSet) override;
-	virtual FTypedElementHandle GetSelectionElement(const TTypedElement<UTypedElementSelectionInterface>& InElementSelectionHandle, const UTypedElementList* InCurrentSelection, const ETypedElementSelectionMethod InSelectionMethod) override;
-	virtual void GetNormalizedElements(const TTypedElement<UTypedElementSelectionInterface>& InElementSelectionHandle, const UTypedElementList* InSelectionSet, const FTypedElementSelectionNormalizationOptions& InNormalizationOptions, UTypedElementList* OutNormalizedElements) override;
+	virtual bool SelectElement(const TTypedElement<UTypedElementSelectionInterface>& InElementSelectionHandle, FTypedElementListRef InSelectionSet, const FTypedElementSelectionOptions& InSelectionOptions) override;
+	virtual bool DeselectElement(const TTypedElement<UTypedElementSelectionInterface>& InElementSelectionHandle, FTypedElementListRef InSelectionSet, const FTypedElementSelectionOptions& InSelectionOptions) override;
+	virtual bool AllowSelectionModifiers(const TTypedElement<UTypedElementSelectionInterface>& InElementSelectionHandle, FTypedElementListConstRef InSelectionSet) override;
+	virtual FTypedElementHandle GetSelectionElement(const TTypedElement<UTypedElementSelectionInterface>& InElementSelectionHandle, FTypedElementListConstRef InCurrentSelection, const ETypedElementSelectionMethod InSelectionMethod) override;
+	virtual void GetNormalizedElements(const TTypedElement<UTypedElementSelectionInterface>& InElementSelectionHandle, FTypedElementListConstRef InSelectionSet, const FTypedElementSelectionNormalizationOptions& InNormalizationOptions, FTypedElementListRef OutNormalizedElements) override;
 
 	bool CanSelectActorElement(const TTypedElement<UTypedElementSelectionInterface>& InActorSelectionHandle, const FTypedElementSelectionOptions& InSelectionOptions) const;
 	bool CanDeselectActorElement(const TTypedElement<UTypedElementSelectionInterface>& InActorSelectionHandle, const FTypedElementSelectionOptions& InSelectionOptions) const;
 
-	bool SelectActorElement(const TTypedElement<UTypedElementSelectionInterface>& InActorSelectionHandle, UTypedElementList* InSelectionSet, const FTypedElementSelectionOptions& InSelectionOptions);
-	bool DeselectActorElement(const TTypedElement<UTypedElementSelectionInterface>& InActorSelectionHandle, UTypedElementList* InSelectionSet, const FTypedElementSelectionOptions& InSelectionOptions);
+	bool SelectActorElement(const TTypedElement<UTypedElementSelectionInterface>& InActorSelectionHandle, FTypedElementListRef InSelectionSet, const FTypedElementSelectionOptions& InSelectionOptions);
+	bool DeselectActorElement(const TTypedElement<UTypedElementSelectionInterface>& InActorSelectionHandle, FTypedElementListRef InSelectionSet, const FTypedElementSelectionOptions& InSelectionOptions);
 
-	bool SelectActorGroup(AGroupActor* InGroupActor, UTypedElementList* InSelectionSet, const FTypedElementSelectionOptions& InSelectionOptions, const bool bForce);
-	bool DeselectActorGroup(AGroupActor* InGroupActor, UTypedElementList* InSelectionSet, const FTypedElementSelectionOptions& InSelectionOptions, const bool bForce);
+	bool SelectActorGroup(AGroupActor* InGroupActor, FTypedElementListRef InSelectionSet, const FTypedElementSelectionOptions& InSelectionOptions, const bool bForce);
+	bool DeselectActorGroup(AGroupActor* InGroupActor, FTypedElementListRef InSelectionSet, const FTypedElementSelectionOptions& InSelectionOptions, const bool bForce);
 
-	static void AppendNormalizedActors(AActor* InActor, const UTypedElementList* InSelectionSet, const FTypedElementSelectionNormalizationOptions& InNormalizationOptions, UTypedElementList* OutNormalizedElements);
+	static void AppendNormalizedActors(AActor* InActor, FTypedElementListConstRef InSelectionSet, const FTypedElementSelectionNormalizationOptions& InNormalizationOptions, FTypedElementListRef OutNormalizedElements);
 };
diff --git a/Engine/Source/Editor/LevelEditor/Private/Elements/Component/ComponentElementLevelEditorSelectionCustomization.cpp b/Engine/Source/Editor/LevelEditor/Private/Elements/Component/ComponentElementLevelEditorSelectionCustomization.cpp
--- a/Engine/Source/Editor/LevelEditor/Private/Elements/Component/ComponentElementLevelEditorSelectionCustomization.cpp
+++ b/Engine/Source/Editor/LevelEditor/Private/Elements/Component/ComponentElementLevelEditorSelectionCustomization.cpp
@@ -26,17 +26,17 @@ bool FComponentElementLevelEditorSelectionCustomization::CanDeselectElement(cons
 	return CanDeselectComponentElement(InElementSelectionHandle, InSelectionOptions);
 }
 
-bool FComponentElementLevelEditorSelectionCustomization::SelectElement(const TTypedElement<UTypedElementSelectionInterface>& InElementSelectionHandle, UTypedElementList* InSelectionSet, const FTypedElementSelectionOptions& InSelectionOptions)
+bool FComponentElementLevelEditorSelectionCustomization::SelectElement(const TTypedElement<UTypedElementSelectionInterface>& InElementSelectionHandle, FTypedElementListRef InSelectionSet, const FTypedElementSelectionOptions& InSelectionOptions)
 {
 	return SelectComponentElement(InElementSelectionHandle, InSelectionSet, InSelectionOptions);
 }
 
-bool FComponentElementLevelEditorSelectionCustomization::DeselectElement(const TTypedElement<UTypedElementSelectionInterface>& InElementSelectionHandle, UTypedElementList* InSelectionSet, const FTypedElementSelectionOptions& InSelectionOptions)
+bool FComponentElementLevelEditorSelectionCustomization::DeselectElement(const TTypedElement<UTypedElementSelectionInterface>& InElementSelectionHandle, FTypedElementListRef InSelectionSet, const FTypedElementSelectionOptions& InSelectionOptions)
 {
 	return DeselectComponentElement(InElementSelectionHandle, InSelectionSet, InSelectionOptions);
 }
 
-FTypedElementHandle FComponentElementLevelEditorSelectionCustomization::GetSelectionElement(const TTypedElement<UTypedElementSelectionInterface>& InElementSelectionHandle, const UTypedElementList* InCurrentSelection, const ETypedElementSelectionMethod InSelectionMethod)
+FTypedElementHandle FComponentElementLevelEditorSelectionCustomization::GetSelectionElement(const TTypedElement<UTypedElementSelectionInterface>& InElementSelectionHandle, FTypedElementListConstRef InCurrentSelection, const ETypedElementSelectionMethod InSelectionMethod)
 {
 	if (const UActorComponent* Component = ComponentElementDataUtil::GetComponentFromHandle(InElementSelectionHandle))
 	{
@@ -83,7 +83,7 @@ FTypedElementHandle FComponentElementLevelEditorSelectionCustomization::GetSelec
 	return InElementSelectionHandle;
 }
 
-void FComponentElementLevelEditorSelectionCustomization::GetNormalizedElements(const TTypedElement<UTypedElementSelectionInterface>& InElementSelectionHandle, const UTypedElementList* InSelectionSet, const FTypedElementSelectionNormalizationOptions& InNormalizationOptions, UTypedElementList* OutNormalizedElements)
+void FComponentElementLevelEditorSelectionCustomization::GetNormalizedElements(const TTypedElement<UTypedElementSelectionInterface>& InElementSelectionHandle, FTypedElementListConstRef InSelectionSet, const FTypedElementSelectionNormalizationOptions& InNormalizationOptions, FTypedElementListRef OutNormalizedElements)
 {
 	UActorComponent* Component = ComponentElementDataUtil::GetComponentFromHandleChecked(InElementSelectionHandle);
 
@@ -142,7 +142,7 @@ bool FComponentElementLevelEditorSelectionCustomization::CanDeselectComponentEle
 	return !GEdSelectionLock;
 }
 
-bool FComponentElementLevelEditorSelectionCustomization::SelectComponentElement(const TTypedElement<UTypedElementSelectionInterface>& InComponentSelectionHandle, UTypedElementList* InSelectionSet, const FTypedElementSelectionOptions& InSelectionOptions)
+bool FComponentElementLevelEditorSelectionCustomization::SelectComponentElement(const TTypedElement<UTypedElementSelectionInterface>& InComponentSelectionHandle, FTypedElementListRef InSelectionSet, const FTypedElementSelectionOptions& InSelectionOptions)
 {
 	UActorComponent* Component = ComponentElementDataUtil::GetComponentFromHandleChecked(InComponentSelectionHandle);
 
@@ -178,7 +178,7 @@ bool FComponentElementLevelEditorSelectionCustomization::SelectComponentElement(
 	return true;
 }
 
-bool FComponentElementLevelEditorSelectionCustomization::DeselectComponentElement(const TTypedElement<UTypedElementSelectionInterface>& InComponentSelectionHandle, UTypedElementList* InSelectionSet, const FTypedElementSelectionOptions& InSelectionOptions)
+bool FComponentElementLevelEditorSelectionCustomization::DeselectComponentElement(const TTypedElement<UTypedElementSelectionInterface>& InComponentSelectionHandle, FTypedElementListRef InSelectionSet, const FTypedElementSelectionOptions& InSelectionOptions)
 {
 	UActorComponent* Component = ComponentElementDataUtil::GetComponentFromHandleChecked(InComponentSelectionHandle);
 
diff --git a/Engine/Source/Editor/LevelEditor/Private/Elements/Component/ComponentElementLevelEditorSelectionCustomization.h b/Engine/Source/Editor/LevelEditor/Private/Elements/Component/ComponentElementLevelEditorSelectionCustomization.h
--- a/Engine/Source/Editor/LevelEditor/Private/Elements/Component/ComponentElementLevelEditorSelectionCustomization.h
+++ b/Engine/Source/Editor/LevelEditor/Private/Elements/Component/ComponentElementLevelEditorSelectionCustomization.h
@@ -10,14 +10,14 @@ class FComponentElementLevelEditorSelectionCustomization : public FTypedElementS
 public:
 	virtual bool CanSelectElement(const TTypedElement<UTypedElementSelectionInterface>& InElementSelectionHandle, const FTypedElementSelectionOptions& InSelectionOptions) override;
 	virtual bool CanDeselectElement(const TTypedElement<UTypedElementSelectionInterface>& InElementSelectionHandle, const FTypedElementSelectionOptions& InSelectionOptions) override;
-	virtual bool SelectElement(const TTypedElement<UTypedElementSelectionInterface>& InElementSelectionHandle, UTypedElementList* InSelectionSet, const FTypedElementSelectionOptions& InSelectionOptions) override;
-	virtual bool DeselectElement(const TTypedElement<UTypedElementSelectionInterface>& InElementSelectionHandle, UTypedElementList* InSelectionSet, const FTypedElementSelectionOptions& InSelectionOptions) override;
-	virtual FTypedElementHandle GetSelectionElement(const TTypedElement<UTypedElementSelectionInterface>& InElementSelectionHandle, const UTypedElementList* InCurrentSelection, const ETypedElementSelectionMethod InSelectionMethod) override;
-	virtual void GetNormalizedElements(const TTypedElement<UTypedElementSelectionInterface>& InElementSelectionHandle, const UTypedElementList* InSelectionSet, const FTypedElementSelectionNormalizationOptions& InNormalizationOptions, UTypedElementList* OutNormalizedElements) override;
+	virtual bool SelectElement(const TTypedElement<UTypedElementSelectionInterface>& InElementSelectionHandle, FTypedElementListRef InSelectionSet, const FTypedElementSelectionOptions& InSelectionOptions) override;
+	virtual bool DeselectElement(const TTypedElement<UTypedElementSelectionInterface>& InElementSelectionHandle, FTypedElementListRef InSelectionSet, const FTypedElementSelectionOptions& InSelectionOptions) override;
+	virtual FTypedElementHandle GetSelectionElement(const TTypedElement<UTypedElementSelectionInterface>& InElementSelectionHandle, FTypedElementListConstRef InCurrentSelection, const ETypedElementSelectionMethod InSelectionMethod) override;
+	virtual void GetNormalizedElements(const TTypedElement<UTypedElementSelectionInterface>& InElementSelectionHandle, FTypedElementListConstRef InSelectionSet, const FTypedElementSelectionNormalizationOptions& InNormalizationOptions, FTypedElementListRef OutNormalizedElements) override;
 
 	bool CanSelectComponentElement(const TTypedElement<UTypedElementSelectionInterface>& InComponentSelectionHandle, const FTypedElementSelectionOptions& InSelectionOptions) const;
 	bool CanDeselectComponentElement(const TTypedElement<UTypedElementSelectionInterface>& InComponentSelectionHandle, const FTypedElementSelectionOptions& InSelectionOptions) const;
 
-	bool SelectComponentElement(const TTypedElement<UTypedElementSelectionInterface>& InComponentSelectionHandle, UTypedElementList* InSelectionSet, const FTypedElementSelectionOptions& InSelectionOptions);
-	bool DeselectComponentElement(const TTypedElement<UTypedElementSelectionInterface>& InComponentSelectionHandle, UTypedElementList* InSelectionSet, const FTypedElementSelectionOptions& InSelectionOptions);
+	bool SelectComponentElement(const TTypedElement<UTypedElementSelectionInterface>& InComponentSelectionHandle, FTypedElementListRef InSelectionSet, const FTypedElementSelectionOptions& InSelectionOptions);
+	bool DeselectComponentElement(const TTypedElement<UTypedElementSelectionInterface>& InComponentSelectionHandle, FTypedElementListRef InSelectionSet, const FTypedElementSelectionOptions& InSelectionOptions);
 };
diff --git a/Engine/Source/Editor/LevelEditor/Private/Elements/SMInstance/SMInstanceElementLevelEditorSelectionCustomization.cpp b/Engine/Source/Editor/LevelEditor/Private/Elements/SMInstance/SMInstanceElementLevelEditorSelectionCustomization.cpp
--- a/Engine/Source/Editor/LevelEditor/Private/Elements/SMInstance/SMInstanceElementLevelEditorSelectionCustomization.cpp
+++ b/Engine/Source/Editor/LevelEditor/Private/Elements/SMInstance/SMInstanceElementLevelEditorSelectionCustomization.cpp
@@ -33,7 +33,7 @@ bool FSMInstanceElementLevelEditorSelectionCustomization::CanDeselectElement(con
 	return !GEdSelectionLock;
 }
 
-bool FSMInstanceElementLevelEditorSelectionCustomization::SelectElement(const TTypedElement<UTypedElementSelectionInterface>& InElementSelectionHandle, UTypedElementList* InSelectionSet, const FTypedElementSelectionOptions& InSelectionOptions)
+bool FSMInstanceElementLevelEditorSelectionCustomization::SelectElement(const TTypedElement<UTypedElementSelectionInterface>& InElementSelectionHandle, FTypedElementListRef InSelectionSet, const FTypedElementSelectionOptions& InSelectionOptions)
 {
 	const FSMInstanceManager SMInstance = SMInstanceElementDataUtil::GetSMInstanceFromHandleChecked(InElementSelectionHandle);
 	if (!SMInstance)
@@ -51,7 +51,7 @@ bool FSMInstanceElementLevelEditorSelectionCustomization::SelectElement(const TT
 	return true;
 }
 
-bool FSMInstanceElementLevelEditorSelectionCustomization::DeselectElement(const TTypedElement<UTypedElementSelectionInterface>& InElementSelectionHandle, UTypedElementList* InSelectionSet, const FTypedElementSelectionOptions& InSelectionOptions)
+bool FSMInstanceElementLevelEditorSelectionCustomization::DeselectElement(const TTypedElement<UTypedElementSelectionInterface>& InElementSelectionHandle, FTypedElementListRef InSelectionSet, const FTypedElementSelectionOptions& InSelectionOptions)
 {
 	const FSMInstanceManager SMInstance = SMInstanceElementDataUtil::GetSMInstanceFromHandleChecked(InElementSelectionHandle);
 	if (!SMInstance)
diff --git a/Engine/Source/Editor/LevelEditor/Private/Elements/SMInstance/SMInstanceElementLevelEditorSelectionCustomization.h b/Engine/Source/Editor/LevelEditor/Private/Elements/SMInstance/SMInstanceElementLevelEditorSelectionCustomization.h
--- a/Engine/Source/Editor/LevelEditor/Private/Elements/SMInstance/SMInstanceElementLevelEditorSelectionCustomization.h
+++ b/Engine/Source/Editor/LevelEditor/Private/Elements/SMInstance/SMInstanceElementLevelEditorSelectionCustomization.h
@@ -10,6 +10,6 @@ class FSMInstanceElementLevelEditorSelectionCustomization : public FTypedElement
 public:
 	virtual bool CanSelectElement(const TTypedElement<UTypedElementSelectionInterface>& InElementSelectionHandle, const FTypedElementSelectionOptions& InSelectionOptions) override;
 	virtual bool CanDeselectElement(const TTypedElement<UTypedElementSelectionInterface>& InElementSelectionHandle, const FTypedElementSelectionOptions& InSelectionOptions) override;
-	virtual bool SelectElement(const TTypedElement<UTypedElementSelectionInterface>& InElementSelectionHandle, UTypedElementList* InSelectionSet, const FTypedElementSelectionOptions& InSelectionOptions) override;
-	virtual bool DeselectElement(const TTypedElement<UTypedElementSelectionInterface>& InElementSelectionHandle, UTypedElementList* InSelectionSet, const FTypedElementSelectionOptions& InSelectionOptions) override;
+	virtual bool SelectElement(const TTypedElement<UTypedElementSelectionInterface>& InElementSelectionHandle, FTypedElementListRef InSelectionSet, const FTypedElementSelectionOptions& InSelectionOptions) override;
+	virtual bool DeselectElement(const TTypedElement<UTypedElementSelectionInterface>& InElementSelectionHandle, FTypedElementListRef InSelectionSet, const FTypedElementSelectionOptions& InSelectionOptions) override;
 };
diff --git a/Engine/Source/Editor/UnrealEd/Classes/Editor/GroupActor.h b/Engine/Source/Editor/UnrealEd/Classes/Editor/GroupActor.h
--- a/Engine/Source/Editor/UnrealEd/Classes/Editor/GroupActor.h
+++ b/Engine/Source/Editor/UnrealEd/Classes/Editor/GroupActor.h
@@ -9,11 +9,11 @@
 #include "CoreMinimal.h"
 #include "UObject/ObjectMacros.h"
 #include "GameFramework/Actor.h"
+#include "Elements/Framework/TypedElementListFwd.h"
 #include "GroupActor.generated.h"
 
 class FLevelEditorViewportClient;
 class FViewport;
-class UTypedElementList;
 class UTypedElementSelectionSet;
 
 UCLASS(hidedropdown,MinimalAPI, notplaceable, NotBlueprintable)
@@ -228,6 +228,6 @@ namespace GroupActorHelpers
 {
 
 UNREALED_API bool ActorHasParentInGroup(const AActor* Actor, const AGroupActor* GroupActor);
-UNREALED_API bool ActorHasParentInSelection(const AActor* Actor, const UTypedElementList* SelectionSet);
+UNREALED_API bool ActorHasParentInSelection(const AActor* Actor, FTypedElementListConstRef SelectionSet);
 
 } // namespace GroupActorHelpers
diff --git a/Engine/Source/Editor/UnrealEd/Private/EditorActor.cpp b/Engine/Source/Editor/UnrealEd/Private/EditorActor.cpp
--- a/Engine/Source/Editor/UnrealEd/Private/EditorActor.cpp
+++ b/Engine/Source/Editor/UnrealEd/Private/EditorActor.cpp
@@ -788,7 +788,7 @@ bool UUnrealEdEngine::DeleteComponents(const TArray<UActorComponent*>& InCompone
 		return false;
 	}
 
-	TUniquePtr<FTypedElementListLegacySyncScopedBatch> LegacySyncBatch = MakeUnique<FTypedElementListLegacySyncScopedBatch>(InSelectionSet->GetElementList());
+	TUniquePtr<FTypedElementListLegacySyncScopedBatch> LegacySyncBatch = MakeUnique<FTypedElementListLegacySyncScopedBatch>(*InSelectionSet->GetElementList());
 
 	const FTypedElementSelectionOptions SelectionOptions = FTypedElementSelectionOptions()
 		.SetAllowHidden(true)
@@ -922,7 +922,7 @@ bool UUnrealEdEngine::DeleteActors(const TArray<AActor*>& InActorsToDelete, UWor
 	bool bRequestedDeleteAllBySoftReference = false;
 	int32 DeleteCount = 0;
 
-	TUniquePtr<FTypedElementListLegacySyncScopedBatch> LegacySyncBatch = MakeUnique<FTypedElementListLegacySyncScopedBatch>(InSelectionSet->GetElementList());
+	TUniquePtr<FTypedElementListLegacySyncScopedBatch> LegacySyncBatch = MakeUnique<FTypedElementListLegacySyncScopedBatch>(*InSelectionSet->GetElementList());
 
 	const FTypedElementSelectionOptions SelectionOptions = FTypedElementSelectionOptions()
 		.SetAllowHidden(true)
diff --git a/Engine/Source/Editor/UnrealEd/Private/EditorSelectUtils.cpp b/Engine/Source/Editor/UnrealEd/Private/EditorSelectUtils.cpp
--- a/Engine/Source/Editor/UnrealEd/Private/EditorSelectUtils.cpp
+++ b/Engine/Source/Editor/UnrealEd/Private/EditorSelectUtils.cpp
@@ -206,7 +206,7 @@ void UUnrealEdEngine::SetPivot( FVector NewPivot, bool bSnapPivotToGrid, bool bI
 	// Pick a new common pivot, or not.
 	TTypedElement<UTypedElementWorldInterface> SingleWorldElement;
 
-	const UTypedElementList* ElementsToManipulate = GCurrentLevelEditingViewportClient->GetElementsToManipulate();
+	FTypedElementListConstRef ElementsToManipulate = GCurrentLevelEditingViewportClient->GetElementsToManipulate();
 	ElementsToManipulate->ForEachElement<UTypedElementWorldInterface>([&NumElements, &TranslateRotateWidgetWorldXAxis, &Widget2DWorldXAxis, &SingleWorldElement](const TTypedElement<UTypedElementWorldInterface>& InWorldElement)
 	{
 		if (NumElements == 0)
@@ -328,7 +328,7 @@ void UUnrealEdEngine::UpdatePivotLocationForSelection( bool bOnChange )
 	// Pick a new common pivot, or not.
 	TTypedElement<UTypedElementWorldInterface> SingleWorldElement;
 	
-	const UTypedElementList* ElementsToManipulate = GCurrentLevelEditingViewportClient->GetElementsToManipulate();
+	FTypedElementListConstRef ElementsToManipulate = GCurrentLevelEditingViewportClient->GetElementsToManipulate();
 	ElementsToManipulate->ForEachElement<UTypedElementWorldInterface>([&SingleWorldElement](const TTypedElement<UTypedElementWorldInterface>& InWorldElement)
 	{
 #if DO_CHECK
diff --git a/Engine/Source/Editor/UnrealEd/Private/Elements/Actor/ActorElementEditorSelectionInterface.cpp b/Engine/Source/Editor/UnrealEd/Private/Elements/Actor/ActorElementEditorSelectionInterface.cpp
--- a/Engine/Source/Editor/UnrealEd/Private/Elements/Actor/ActorElementEditorSelectionInterface.cpp
+++ b/Engine/Source/Editor/UnrealEd/Private/Elements/Actor/ActorElementEditorSelectionInterface.cpp
@@ -38,10 +38,11 @@ private:
 	TWeakObjectPtr<AActor> ActorPtr;
 };
 
-bool UActorElementEditorSelectionInterface::IsElementSelected(const FTypedElementHandle& InElementHandle, const UTypedElementList* InSelectionSet, const FTypedElementIsSelectedOptions& InSelectionOptions)
+bool UActorElementEditorSelectionInterface::IsElementSelected(const FTypedElementHandle& InElementHandle, const FTypedElementListProxy InSelectionSet, const FTypedElementIsSelectedOptions& InSelectionOptions)
 {
+	FTypedElementListConstPtr SelectionSetPtr = InSelectionSet.GetElementList();
 	const AActor* Actor = ActorElementDataUtil::GetActorFromHandle(InElementHandle);
-	return Actor && IsActorSelected(Actor, InSelectionSet, InSelectionOptions);
+	return SelectionSetPtr && Actor && IsActorSelected(Actor, SelectionSetPtr.ToSharedRef(), InSelectionOptions);
 }
 
 bool UActorElementEditorSelectionInterface::ShouldPreventTransactions(const FTypedElementHandle& InElementHandle)
@@ -55,7 +56,7 @@ TUniquePtr<ITypedElementTransactedElement> UActorElementEditorSelectionInterface
 	return MakeUnique<FActorElementTransactedElement>();
 }
 
-bool UActorElementEditorSelectionInterface::IsActorSelected(const AActor* InActor, const UTypedElementList* InSelectionSet, const FTypedElementIsSelectedOptions& InSelectionOptions)
+bool UActorElementEditorSelectionInterface::IsActorSelected(const AActor* InActor, FTypedElementListConstRef InSelectionSet, const FTypedElementIsSelectedOptions& InSelectionOptions)
 {
 	if (InSelectionSet->Num() == 0)
 	{
diff --git a/Engine/Source/Editor/UnrealEd/Private/Elements/Actor/ActorElementEditorSelectionInterface.h b/Engine/Source/Editor/UnrealEd/Private/Elements/Actor/ActorElementEditorSelectionInterface.h
--- a/Engine/Source/Editor/UnrealEd/Private/Elements/Actor/ActorElementEditorSelectionInterface.h
+++ b/Engine/Source/Editor/UnrealEd/Private/Elements/Actor/ActorElementEditorSelectionInterface.h
@@ -13,9 +13,9 @@ class UActorElementEditorSelectionInterface : public UActorElementSelectionInter
 	GENERATED_BODY()
 
 public:
-	virtual bool IsElementSelected(const FTypedElementHandle& InElementHandle, const UTypedElementList* InSelectionSet, const FTypedElementIsSelectedOptions& InSelectionOptions) override;
+	virtual bool IsElementSelected(const FTypedElementHandle& InElementHandle, const FTypedElementListProxy InSelectionSet, const FTypedElementIsSelectedOptions& InSelectionOptions) override;
 	virtual bool ShouldPreventTransactions(const FTypedElementHandle& InElementHandle) override;
 	virtual TUniquePtr<ITypedElementTransactedElement> CreateTransactedElementImpl() override;
 
-	static bool IsActorSelected(const AActor* InActor, const UTypedElementList* InSelectionSet, const FTypedElementIsSelectedOptions& InSelectionOptions);
+	static bool IsActorSelected(const AActor* InActor, FTypedElementListConstRef InSelectionSet, const FTypedElementIsSelectedOptions& InSelectionOptions);
 };
diff --git a/Engine/Source/Editor/UnrealEd/Private/Elements/Component/ComponentElementEditorSelectionInterface.cpp b/Engine/Source/Editor/UnrealEd/Private/Elements/Component/ComponentElementEditorSelectionInterface.cpp
--- a/Engine/Source/Editor/UnrealEd/Private/Elements/Component/ComponentElementEditorSelectionInterface.cpp
+++ b/Engine/Source/Editor/UnrealEd/Private/Elements/Component/ComponentElementEditorSelectionInterface.cpp
@@ -39,10 +39,11 @@ private:
 	TWeakObjectPtr<UActorComponent> ComponentPtr;
 };
 
-bool UComponentElementEditorSelectionInterface::IsElementSelected(const FTypedElementHandle& InElementHandle, const UTypedElementList* InSelectionSet, const FTypedElementIsSelectedOptions& InSelectionOptions)
+bool UComponentElementEditorSelectionInterface::IsElementSelected(const FTypedElementHandle& InElementHandle, const FTypedElementListProxy InSelectionSet, const FTypedElementIsSelectedOptions& InSelectionOptions)
 {
+	FTypedElementListConstPtr SelectionSetPtr = InSelectionSet.GetElementList();
 	const UActorComponent* Component = ComponentElementDataUtil::GetComponentFromHandle(InElementHandle);
-	return Component && IsComponentSelected(Component, InSelectionSet, InSelectionOptions);
+	return SelectionSetPtr && Component && IsComponentSelected(Component, SelectionSetPtr.ToSharedRef(), InSelectionOptions);
 }
 
 bool UComponentElementEditorSelectionInterface::ShouldPreventTransactions(const FTypedElementHandle& InElementHandle)
@@ -56,7 +57,7 @@ TUniquePtr<ITypedElementTransactedElement> UComponentElementEditorSelectionInter
 	return MakeUnique<FComponentElementTransactedElement>();
 }
 
-bool UComponentElementEditorSelectionInterface::IsComponentSelected(const UActorComponent* InComponent, const UTypedElementList* InSelectionSet, const FTypedElementIsSelectedOptions& InSelectionOptions)
+bool UComponentElementEditorSelectionInterface::IsComponentSelected(const UActorComponent* InComponent, FTypedElementListConstRef InSelectionSet, const FTypedElementIsSelectedOptions& InSelectionOptions)
 {
 	if (InSelectionSet->Num() == 0)
 	{
diff --git a/Engine/Source/Editor/UnrealEd/Private/Elements/Component/ComponentElementEditorSelectionInterface.h b/Engine/Source/Editor/UnrealEd/Private/Elements/Component/ComponentElementEditorSelectionInterface.h
--- a/Engine/Source/Editor/UnrealEd/Private/Elements/Component/ComponentElementEditorSelectionInterface.h
+++ b/Engine/Source/Editor/UnrealEd/Private/Elements/Component/ComponentElementEditorSelectionInterface.h
@@ -13,9 +13,9 @@ class UComponentElementEditorSelectionInterface : public UComponentElementSelect
 	GENERATED_BODY()
 
 public:
-	virtual bool IsElementSelected(const FTypedElementHandle& InElementHandle, const UTypedElementList* InSelectionSet, const FTypedElementIsSelectedOptions& InSelectionOptions) override;
+	virtual bool IsElementSelected(const FTypedElementHandle& InElementHandle, const FTypedElementListProxy InSelectionSet, const FTypedElementIsSelectedOptions& InSelectionOptions) override;
 	virtual bool ShouldPreventTransactions(const FTypedElementHandle& InElementHandle) override;
 	virtual TUniquePtr<ITypedElementTransactedElement> CreateTransactedElementImpl() override;
 
-	static bool IsComponentSelected(const UActorComponent* InComponent, const UTypedElementList* InSelectionSet, const FTypedElementIsSelectedOptions& InSelectionOptions);
+	static bool IsComponentSelected(const UActorComponent* InComponent, FTypedElementListConstRef InSelectionSet, const FTypedElementIsSelectedOptions& InSelectionOptions);
 };
diff --git a/Engine/Source/Editor/UnrealEd/Private/Elements/SMInstance/SMInstanceElementEditorSelectionInterface.cpp b/Engine/Source/Editor/UnrealEd/Private/Elements/SMInstance/SMInstanceElementEditorSelectionInterface.cpp
--- a/Engine/Source/Editor/UnrealEd/Private/Elements/SMInstance/SMInstanceElementEditorSelectionInterface.cpp
+++ b/Engine/Source/Editor/UnrealEd/Private/Elements/SMInstance/SMInstanceElementEditorSelectionInterface.cpp
@@ -41,16 +41,19 @@ private:
 	uint64 InstanceId = 0;
 };
 
-bool USMInstanceElementEditorSelectionInterface::IsElementSelected(const FTypedElementHandle& InElementHandle, const UTypedElementList* InSelectionSet, const FTypedElementIsSelectedOptions& InSelectionOptions)
+bool USMInstanceElementEditorSelectionInterface::IsElementSelected(const FTypedElementHandle& InElementHandle, const FTypedElementListProxy InSelectionSet, const FTypedElementIsSelectedOptions& InSelectionOptions)
 {
-	if (const FSMInstanceManager SMInstance = SMInstanceElementDataUtil::GetSMInstanceFromHandle(InElementHandle))
+	FTypedElementListConstPtr SelectionSetPtr = InSelectionSet.GetElementList();
+	const FSMInstanceManager SMInstance = SMInstanceElementDataUtil::GetSMInstanceFromHandle(InElementHandle);
+
+	if (SelectionSetPtr && SMInstance)
 	{
-		if (InSelectionSet->Num() == 0)
+		if (SelectionSetPtr->Num() == 0)
 		{
 			return false;
 		}
 
-		if (InSelectionSet->Contains(InElementHandle))
+		if (SelectionSetPtr->Contains(InElementHandle))
 		{
 			return true;
 		}
@@ -59,7 +62,7 @@ bool USMInstanceElementEditorSelectionInterface::IsElementSelected(const FTypedE
 		{
 			if (FTypedElementHandle ISMComponentElement = UEngineElementsLibrary::AcquireEditorComponentElementHandle(SMInstance.GetISMComponent(), /*bAllowCreate*/false))
 			{
-				return InSelectionSet->Contains(ISMComponentElement);
+				return SelectionSetPtr->Contains(ISMComponentElement);
 			}
 		}
 	}
diff --git a/Engine/Source/Editor/UnrealEd/Private/Elements/SMInstance/SMInstanceElementEditorSelectionInterface.h b/Engine/Source/Editor/UnrealEd/Private/Elements/SMInstance/SMInstanceElementEditorSelectionInterface.h
--- a/Engine/Source/Editor/UnrealEd/Private/Elements/SMInstance/SMInstanceElementEditorSelectionInterface.h
+++ b/Engine/Source/Editor/UnrealEd/Private/Elements/SMInstance/SMInstanceElementEditorSelectionInterface.h
@@ -11,7 +11,7 @@ class USMInstanceElementEditorSelectionInterface : public USMInstanceElementSele
 	GENERATED_BODY()
 
 public:
-	virtual bool IsElementSelected(const FTypedElementHandle& InElementHandle, const UTypedElementList* InSelectionSet, const FTypedElementIsSelectedOptions& InSelectionOptions) override;
+	virtual bool IsElementSelected(const FTypedElementHandle& InElementHandle, const FTypedElementListProxy InSelectionSet, const FTypedElementIsSelectedOptions& InSelectionOptions) override;
 	virtual bool ShouldPreventTransactions(const FTypedElementHandle& InElementHandle) override;
 	virtual TUniquePtr<ITypedElementTransactedElement> CreateTransactedElementImpl() override;
 };
diff --git a/Engine/Source/Editor/UnrealEd/Private/Elements/SMInstance/SMInstanceElementEditorWorldInterface.cpp b/Engine/Source/Editor/UnrealEd/Private/Elements/SMInstance/SMInstanceElementEditorWorldInterface.cpp
--- a/Engine/Source/Editor/UnrealEd/Private/Elements/SMInstance/SMInstanceElementEditorWorldInterface.cpp
+++ b/Engine/Source/Editor/UnrealEd/Private/Elements/SMInstance/SMInstanceElementEditorWorldInterface.cpp
@@ -46,7 +46,7 @@ bool USMInstanceElementEditorWorldInterface::DeleteElements(TArrayView<const FTy
 
 		bool bDidDelete = false;
 
-		FTypedElementListLegacySyncScopedBatch LegacySyncBatch(InSelectionSet->GetElementList());
+		FTypedElementListLegacySyncScopedBatch LegacySyncBatch(*InSelectionSet->GetElementList());
 		for (TTuple<TScriptInterface<ISMInstanceManager>, TArray<FSMInstanceId>>& BatchedInstancesToDeletePair : BatchedInstancesToDelete)
 		{
 			for (const FSMInstanceId& InstanceId : BatchedInstancesToDeletePair.Value)
diff --git a/Engine/Source/Editor/UnrealEd/Private/GroupActor.cpp b/Engine/Source/Editor/UnrealEd/Private/GroupActor.cpp
--- a/Engine/Source/Editor/UnrealEd/Private/GroupActor.cpp
+++ b/Engine/Source/Editor/UnrealEd/Private/GroupActor.cpp
@@ -147,18 +147,15 @@ bool ActorHasParentInGroup(const AActor* Actor, const AGroupActor* GroupActor)
 	return false;
 }
 
-bool ActorHasParentInSelection(const AActor* Actor, const UTypedElementList* SelectionSet)
+bool ActorHasParentInSelection(const AActor* Actor, FTypedElementListConstRef SelectionSet)
 {
 	check(Actor);
-	if (SelectionSet)
+	for (const AActor* ParentActor = Actor->GetAttachParentActor(); ParentActor; ParentActor = ParentActor->GetAttachParentActor())
 	{
-		for (const AActor* ParentActor = Actor->GetAttachParentActor(); ParentActor; ParentActor = ParentActor->GetAttachParentActor())
+		FTypedElementHandle ParentActorElementHandle = UEngineElementsLibrary::AcquireEditorActorElementHandle(ParentActor, /*bAllowCreate*/false);
+		if (ParentActorElementHandle && SelectionSet->Contains(ParentActorElementHandle))
 		{
-			FTypedElementHandle ParentActorElementHandle = UEngineElementsLibrary::AcquireEditorActorElementHandle(ParentActor, /*bAllowCreate*/false);
-			if (ParentActorElementHandle && SelectionSet->Contains(ParentActorElementHandle))
-			{
-				return true;
-			}
+			return true;
 		}
 	}
 	return false;
diff --git a/Engine/Source/Editor/UnrealEd/Private/LevelEditorViewport.cpp b/Engine/Source/Editor/UnrealEd/Private/LevelEditorViewport.cpp
--- a/Engine/Source/Editor/UnrealEd/Private/LevelEditorViewport.cpp
+++ b/Engine/Source/Editor/UnrealEd/Private/LevelEditorViewport.cpp
@@ -2127,7 +2127,7 @@ void FLevelEditorViewportClient::ProcessClick(FSceneView& View, HHitProxy* HitPr
 	// otherwise the usual call in TrackingStopped would include the newly selected element
 	if (bHasBegunGizmoManipulation)
 	{
-		const UTypedElementList* ElementsToManipulate = GetElementsToManipulate();
+		FTypedElementListConstRef ElementsToManipulate = GetElementsToManipulate();
 		ViewportInteraction->EndGizmoManipulation(ElementsToManipulate, GetWidgetMode(), ETypedElementViewportInteractionGizmoManipulationType::Click);
 		bHasBegunGizmoManipulation = false;
 	}
@@ -2592,12 +2592,11 @@ bool FLevelEditorViewportClient::InputWidgetDelta(FViewport* InViewport, EAxisLi
 							UTypedElementCommonActions* CommonActions = LevelEditor ? LevelEditor->GetCommonActions() : nullptr;
 							if (CommonActions)
 							{
-								const UTypedElementList* ElementsToManipulate = GetElementsToManipulate();
+								FTypedElementListConstRef ElementsToManipulate = GetElementsToManipulate();
 								TArray<FTypedElementHandle> DuplicatedElements;
 								{
-									TGCObjectScopeGuard<UTypedElementList> ElementsToDuplicate(ElementsToManipulate->Clone()); // Clone this, as the duplicate may triger construction scripts to modify the selection as we duplicate
-									DuplicatedElements = CommonActions->DuplicateNormalizedElements(ElementsToDuplicate.Get(), GetWorld(), FVector::ZeroVector);
-									ElementsToDuplicate.Get()->Empty();
+									FTypedElementListRef ElementsToDuplicate = ElementsToManipulate->Clone(); // Clone this, as the duplicate may triger construction scripts to modify the selection as we duplicate
+									DuplicatedElements = CommonActions->DuplicateNormalizedElements(ElementsToDuplicate, GetWorld(), FVector::ZeroVector);
 								}
 
 								// Exclusively select the new elements, so that future gizmo interaction manipulates those items instead
@@ -2955,7 +2954,7 @@ void FLevelEditorViewportClient::TrackingStarted( const FInputEventState& InInpu
 	{
 		bIsTrackingBrushModification = false;
 
-		const UTypedElementList* ElementsToManipulate = GetElementsToManipulate();
+		FTypedElementListConstRef ElementsToManipulate = GetElementsToManipulate();
 		TypedElementListObjectUtil::ForEachObject<ABrush>(ElementsToManipulate, [this](ABrush* InBrush)
 		{
 			bIsTrackingBrushModification = InBrush && !InBrush->IsVolumeBrush() && !FActorEditorUtils::IsABuilderBrush(InBrush);
@@ -2967,7 +2966,7 @@ void FLevelEditorViewportClient::TrackingStarted( const FInputEventState& InInpu
 	{
 		Widget->SetSnapEnabled(true);
 
-		const UTypedElementList* ElementsToManipulate = GetElementsToManipulate();
+		FTypedElementListConstRef ElementsToManipulate = GetElementsToManipulate();
 		ViewportInteraction->BeginGizmoManipulation(ElementsToManipulate, GetWidgetMode());
 		bHasBegunGizmoManipulation = true;
 
@@ -3094,7 +3093,7 @@ void FLevelEditorViewportClient::TrackingStopped()
 
 		if (bHasBegunGizmoManipulation)
 		{
-			const UTypedElementList* ElementsToManipulate = GetElementsToManipulate();
+			FTypedElementListConstRef ElementsToManipulate = GetElementsToManipulate();
 			ViewportInteraction->EndGizmoManipulation(ElementsToManipulate, GetWidgetMode(), bDidMove ? ETypedElementViewportInteractionGizmoManipulationType::Drag : ETypedElementViewportInteractionGizmoManipulationType::Click);
 			bHasBegunGizmoManipulation = false;
 		}
@@ -3107,7 +3106,7 @@ void FLevelEditorViewportClient::TrackingStopped()
 
 	if (bNeedToRestoreComponentBeingMovedFlag)
 	{
-		const UTypedElementList* ElementsToManipulate = GetElementsToManipulate();
+		FTypedElementListConstRef ElementsToManipulate = GetElementsToManipulate();
 		TypedElementListObjectUtil::ForEachObject<AActor>(ElementsToManipulate, [this](AActor* InActor)
 		{
 			SetActorBeingMovedByEditor(InActor, false);
@@ -3510,7 +3509,7 @@ bool FLevelEditorViewportClient::HaveSelectedObjectsBeenChanged() const
 	return (TrackingTransaction.TransCount > 0 || TrackingTransaction.IsActive()) && (MouseDeltaTracker->HasReceivedDelta() || MouseDeltaTracker->WasExternalMovement());
 }
 
-const UTypedElementList* FLevelEditorViewportClient::GetElementsToManipulate(const bool bForceRefresh)
+FTypedElementListConstRef FLevelEditorViewportClient::GetElementsToManipulate(const bool bForceRefresh)
 {
 	CacheElementsToManipulate(bForceRefresh);
 	return CachedElementsToManipulate;
@@ -3796,7 +3795,7 @@ void FLevelEditorViewportClient::ApplyDeltaToSelectedElements(const FTransform&
 	FInputDeviceState InputState;
 	InputState.SetModifierKeyStates(IsShiftPressed(), IsAltPressed(), IsCtrlPressed(), IsCmdPressed());
 
-	const UTypedElementList* ElementsToManipulate = GetElementsToManipulate();
+	FTypedElementListConstRef ElementsToManipulate = GetElementsToManipulate();
 	ViewportInteraction->UpdateGizmoManipulation(ElementsToManipulate, GetWidgetMode(), Widget ? Widget->GetCurrentAxis() : EAxisList::None, InputState, ModifiedDeltaTransform);
 }
 
@@ -4588,8 +4587,6 @@ void FLevelEditorViewportClient::AddReferencedObjects( FReferenceCollector& Coll
 {
 	FEditorViewportClient::AddReferencedObjects( Collector );
 
-	Collector.AddReferencedObject(CachedElementsToManipulate);
-
 	for( TSet<FViewportHoverTarget>::TIterator It( FLevelEditorViewportClient::HoveredObjects ); It; ++It )
 	{
 		FViewportHoverTarget& CurHoverTarget = *It;
diff --git a/Engine/Source/Editor/UnrealEd/Private/Selection.cpp b/Engine/Source/Editor/UnrealEd/Private/Selection.cpp
--- a/Engine/Source/Editor/UnrealEd/Private/Selection.cpp
+++ b/Engine/Source/Editor/UnrealEd/Private/Selection.cpp
@@ -129,7 +129,7 @@ UObject* USelection::GetSelectedObject(const int32 InIndex) const
 {
 	if (ElementSelectionSet)
 	{
-		const UTypedElementList* ElementList = ElementSelectionSet->GetElementList();
+		FTypedElementListConstRef ElementList = ElementSelectionSet->GetElementList();
 		if (ElementList->IsValidIndex(InIndex))
 		{
 			const FTypedElementHandle ElementHandle = ElementList->GetElementHandleAt(InIndex);
@@ -198,9 +198,9 @@ UObject* USelection::GetObjectForElementHandle(const FTypedElementHandle& InElem
 	return nullptr;
 }
 
-void USelection::OnElementListSyncEvent(const UTypedElementList* InElementList, FTypedElementListLegacySync::ESyncType InSyncType, const FTypedElementHandle& InElementHandle, bool bIsWithinBatchOperation)
+void USelection::OnElementListSyncEvent(const FTypedElementList& InElementList, FTypedElementListLegacySync::ESyncType InSyncType, const FTypedElementHandle& InElementHandle, bool bIsWithinBatchOperation)
 {
-	check(InElementList == ElementSelectionSet->GetElementList());
+	check(&InElementList == &ElementSelectionSet->GetElementList().Get());
 
 	const bool bNotify = !bIsWithinBatchOperation;
 	if (bNotify)
diff --git a/Engine/Source/Editor/UnrealEd/Private/Tools/DefaultEdMode.cpp b/Engine/Source/Editor/UnrealEd/Private/Tools/DefaultEdMode.cpp
--- a/Engine/Source/Editor/UnrealEd/Private/Tools/DefaultEdMode.cpp
+++ b/Engine/Source/Editor/UnrealEd/Private/Tools/DefaultEdMode.cpp
@@ -15,7 +15,7 @@ public:
 	{
 		if (GCurrentLevelEditingViewportClient)
 		{
-			const UTypedElementList* ElementsToManipulate = GCurrentLevelEditingViewportClient->GetElementsToManipulate();
+			FTypedElementListConstRef ElementsToManipulate = GCurrentLevelEditingViewportClient->GetElementsToManipulate();
 			return ElementsToManipulate->Num() > 0;
 		}
 		return false;
diff --git a/Engine/Source/Editor/UnrealEd/Public/AssetEditorModeManager.h b/Engine/Source/Editor/UnrealEd/Public/AssetEditorModeManager.h
--- a/Engine/Source/Editor/UnrealEd/Public/AssetEditorModeManager.h
+++ b/Engine/Source/Editor/UnrealEd/Public/AssetEditorModeManager.h
@@ -6,7 +6,6 @@
 #include "EditorModeManager.h"
 
 class FPreviewScene;
-class UTypedElementList;
 
 //////////////////////////////////////////////////////////////////////////
 // FAssetEditorModeManager
diff --git a/Engine/Source/Editor/UnrealEd/Public/LevelEditorViewport.h b/Engine/Source/Editor/UnrealEd/Public/LevelEditorViewport.h
--- a/Engine/Source/Editor/UnrealEd/Public/LevelEditorViewport.h
+++ b/Engine/Source/Editor/UnrealEd/Public/LevelEditorViewport.h
@@ -11,6 +11,7 @@
 #include "EditorViewportClient.h"
 #include "UObject/ObjectKey.h"
 #include "UnrealEdMisc.h"
+#include "Elements/Framework/TypedElementListFwd.h"
 
 struct FAssetData;
 struct FMinimalViewInfo;
@@ -21,7 +22,6 @@ class ILevelEditor;
 class SLevelViewport;
 class UActorFactory;
 class UModel;
-class UTypedElementList;
 class UTypedElementSelectionSet;
 struct FWorldContext;
 struct FTypedElementHandle;
@@ -331,7 +331,7 @@ public:
 	/**
 	 * Get the elements (from the current selection set) that this viewport can manipulate (eg, via the transform gizmo).
 	 */
-	const UTypedElementList* GetElementsToManipulate(const bool bForceRefresh = false);
+	FTypedElementListConstRef GetElementsToManipulate(const bool bForceRefresh = false);
 
 	virtual void SetIsSimulateInEditorViewport( bool bInIsSimulateInEditorViewport ) override;
 
@@ -910,7 +910,7 @@ private:
 
 	/** The elements (from the current selection set) that this viewport can manipulate (eg, via the transform gizmo) */
 	bool bHasCachedElementsToManipulate = false;
-	UTypedElementList* CachedElementsToManipulate = nullptr;
+	FTypedElementListRef CachedElementsToManipulate;
 
 	/** Bit array representing the visibility of every sprite category in the current viewport */
 	TBitArray<>	SpriteCategoryVisibility;
diff --git a/Engine/Source/Editor/UnrealEd/Public/Selection.h b/Engine/Source/Editor/UnrealEd/Public/Selection.h
--- a/Engine/Source/Editor/UnrealEd/Public/Selection.h
+++ b/Engine/Source/Editor/UnrealEd/Public/Selection.h
@@ -325,7 +325,7 @@ private:
 	bool IsValidObjectToSelect(const UObject* InObject) const;
 	UObject* GetObjectForElementHandle(const FTypedElementHandle& InElementHandle) const;
 
-	void OnElementListSyncEvent(const UTypedElementList* InElementList, FTypedElementListLegacySync::ESyncType InSyncType, const FTypedElementHandle& InElementHandle, bool bIsWithinBatchOperation);
+	void OnElementListSyncEvent(const FTypedElementList& InElementList, FTypedElementListLegacySync::ESyncType InSyncType, const FTypedElementHandle& InElementHandle, bool bIsWithinBatchOperation);
 
 	/** Bridge from UObjects to their corresponding typed elements. */
 	TSharedPtr<ISelectionElementBridge> SelectionElementBridge;
diff --git a/Engine/Source/Runtime/Engine/Private/Elements/Framework/TypedElementCommonActions.cpp b/Engine/Source/Runtime/Engine/Private/Elements/Framework/TypedElementCommonActions.cpp
--- a/Engine/Source/Runtime/Engine/Private/Elements/Framework/TypedElementCommonActions.cpp
+++ b/Engine/Source/Runtime/Engine/Private/Elements/Framework/TypedElementCommonActions.cpp
@@ -20,29 +20,30 @@ void FTypedElementCommonActionsCustomization::DuplicateElements(UTypedElementWor
 
 bool UTypedElementCommonActions::DeleteSelectedElements(UTypedElementSelectionSet* SelectionSet, UWorld* World, const FTypedElementDeletionOptions& DeletionOptions)
 {
-	TGCObjectScopeGuard<UTypedElementList> NormalizedElements(SelectionSet->GetNormalizedSelection(FTypedElementSelectionNormalizationOptions()));
-	const bool bSuccess = DeleteNormalizedElements(NormalizedElements.Get(), World, SelectionSet, DeletionOptions);
-	NormalizedElements.Get()->Reset();
-	return bSuccess;
+	FTypedElementListRef NormalizedElements = SelectionSet->GetNormalizedSelection(FTypedElementSelectionNormalizationOptions());
+	return DeleteNormalizedElements(NormalizedElements, World, SelectionSet, DeletionOptions);
 }
 
-bool UTypedElementCommonActions::DeleteNormalizedElements(const UTypedElementList* ElementList, UWorld* World, UTypedElementSelectionSet* SelectionSet, const FTypedElementDeletionOptions& DeletionOptions)
+bool UTypedElementCommonActions::DeleteNormalizedElements(const FTypedElementListProxy ElementList, UWorld* World, UTypedElementSelectionSet* SelectionSet, const FTypedElementDeletionOptions& DeletionOptions)
 {
-	TMap<FTypedHandleTypeId, TArray<FTypedElementHandle>> ElementsToDeleteByType;
-	TypedElementUtil::BatchElementsByType(ElementList, ElementsToDeleteByType);
-	
 	bool bSuccess = false;
-	
-	UTypedElementRegistry* Registry = UTypedElementRegistry::GetInstance();
-	UTypedElementRegistry::FDisableElementDestructionOnGC GCGuard(Registry);
-	
-	for (const auto& ElementsByTypePair : ElementsToDeleteByType)
+
+	if (FTypedElementListConstPtr ElementListPtr = ElementList.GetElementList())
 	{
-		FTypedElementCommonActionsCustomization* CommonActionsCustomization = GetInterfaceCustomizationByTypeId(ElementsByTypePair.Key);
-		UTypedElementWorldInterface* WorldInterface = Registry->GetElementInterface<UTypedElementWorldInterface>(ElementsByTypePair.Key);
-		if (CommonActionsCustomization && WorldInterface)
+		TMap<FTypedHandleTypeId, TArray<FTypedElementHandle>> ElementsToDeleteByType;
+		TypedElementUtil::BatchElementsByType(ElementListPtr.ToSharedRef(), ElementsToDeleteByType);
+
+		UTypedElementRegistry* Registry = UTypedElementRegistry::GetInstance();
+		UTypedElementRegistry::FDisableElementDestructionOnGC GCGuard(Registry);
+
+		for (const auto& ElementsByTypePair : ElementsToDeleteByType)
 		{
-			bSuccess |= CommonActionsCustomization->DeleteElements(WorldInterface, ElementsByTypePair.Value, World, SelectionSet, DeletionOptions);
+			FTypedElementCommonActionsCustomization* CommonActionsCustomization = GetInterfaceCustomizationByTypeId(ElementsByTypePair.Key);
+			UTypedElementWorldInterface* WorldInterface = Registry->GetElementInterface<UTypedElementWorldInterface>(ElementsByTypePair.Key);
+			if (CommonActionsCustomization && WorldInterface)
+			{
+				bSuccess |= CommonActionsCustomization->DeleteElements(WorldInterface, ElementsByTypePair.Value, World, SelectionSet, DeletionOptions);
+			}
 		}
 	}
 	
@@ -51,28 +52,29 @@ bool UTypedElementCommonActions::DeleteNormalizedElements(const UTypedElementLis
 
 TArray<FTypedElementHandle> UTypedElementCommonActions::DuplicateSelectedElements(const UTypedElementSelectionSet* SelectionSet, UWorld* World, const FVector& LocationOffset)
 {
-	TGCObjectScopeGuard<UTypedElementList> NormalizedElements(SelectionSet->GetNormalizedSelection(FTypedElementSelectionNormalizationOptions()));
-	TArray<FTypedElementHandle> NewElements = DuplicateNormalizedElements(NormalizedElements.Get(), World, LocationOffset);
-	NormalizedElements.Get()->Reset();
-	return NewElements;
+	FTypedElementListRef NormalizedElements = SelectionSet->GetNormalizedSelection(FTypedElementSelectionNormalizationOptions());
+	return DuplicateNormalizedElements(NormalizedElements, World, LocationOffset);
 }
 
-TArray<FTypedElementHandle> UTypedElementCommonActions::DuplicateNormalizedElements(const UTypedElementList* ElementList, UWorld* World, const FVector& LocationOffset)
+TArray<FTypedElementHandle> UTypedElementCommonActions::DuplicateNormalizedElements(const FTypedElementListProxy ElementList, UWorld* World, const FVector& LocationOffset)
 {
-	TMap<FTypedHandleTypeId, TArray<FTypedElementHandle>> ElementsToDuplicateByType;
-	TypedElementUtil::BatchElementsByType(ElementList, ElementsToDuplicateByType);
-	
 	TArray<FTypedElementHandle> NewElements;
-	NewElements.Reserve(ElementList->Num());
-	
-	UTypedElementRegistry* Registry = UTypedElementRegistry::GetInstance();
-	for (const auto& ElementsByTypePair : ElementsToDuplicateByType)
+	if (FTypedElementListConstPtr ElementListPtr = ElementList.GetElementList())
 	{
-		FTypedElementCommonActionsCustomization* CommonActionsCustomization = GetInterfaceCustomizationByTypeId(ElementsByTypePair.Key);
-		UTypedElementWorldInterface* WorldInterface = Registry->GetElementInterface<UTypedElementWorldInterface>(ElementsByTypePair.Key);
-		if (CommonActionsCustomization && WorldInterface)
+		NewElements.Reserve(ElementListPtr->Num());
+
+		TMap<FTypedHandleTypeId, TArray<FTypedElementHandle>> ElementsToDuplicateByType;
+		TypedElementUtil::BatchElementsByType(ElementListPtr.ToSharedRef(), ElementsToDuplicateByType);
+
+		UTypedElementRegistry* Registry = UTypedElementRegistry::GetInstance();
+		for (const auto& ElementsByTypePair : ElementsToDuplicateByType)
 		{
-			CommonActionsCustomization->DuplicateElements(WorldInterface, ElementsByTypePair.Value, World, LocationOffset, NewElements);
+			FTypedElementCommonActionsCustomization* CommonActionsCustomization = GetInterfaceCustomizationByTypeId(ElementsByTypePair.Key);
+			UTypedElementWorldInterface* WorldInterface = Registry->GetElementInterface<UTypedElementWorldInterface>(ElementsByTypePair.Key);
+			if (CommonActionsCustomization && WorldInterface)
+			{
+				CommonActionsCustomization->DuplicateElements(WorldInterface, ElementsByTypePair.Value, World, LocationOffset, NewElements);
+			}
 		}
 	}
 	
diff --git a/Engine/Source/Runtime/Engine/Private/Elements/SMInstance/SMInstanceElementSelectionInterface.cpp b/Engine/Source/Runtime/Engine/Private/Elements/SMInstance/SMInstanceElementSelectionInterface.cpp
--- a/Engine/Source/Runtime/Engine/Private/Elements/SMInstance/SMInstanceElementSelectionInterface.cpp
+++ b/Engine/Source/Runtime/Engine/Private/Elements/SMInstance/SMInstanceElementSelectionInterface.cpp
@@ -5,7 +5,7 @@
 #include "Components/InstancedStaticMeshComponent.h"
 #include "Elements/Framework/TypedElementList.h"
 
-bool USMInstanceElementSelectionInterface::SelectElement(const FTypedElementHandle& InElementHandle, UTypedElementList* InSelectionSet, const FTypedElementSelectionOptions& InSelectionOptions)
+bool USMInstanceElementSelectionInterface::SelectElement(const FTypedElementHandle& InElementHandle, FTypedElementListProxy InSelectionSet, const FTypedElementSelectionOptions& InSelectionOptions)
 {
 	FSMInstanceManager SMInstance = SMInstanceElementDataUtil::GetSMInstanceFromHandle(InElementHandle);
 	if (SMInstance && Super::SelectElement(InElementHandle, InSelectionSet, InSelectionOptions))
@@ -16,7 +16,7 @@ bool USMInstanceElementSelectionInterface::SelectElement(const FTypedElementHand
 	return false;
 }
 
-bool USMInstanceElementSelectionInterface::DeselectElement(const FTypedElementHandle& InElementHandle, UTypedElementList* InSelectionSet, const FTypedElementSelectionOptions& InSelectionOptions)
+bool USMInstanceElementSelectionInterface::DeselectElement(const FTypedElementHandle& InElementHandle, FTypedElementListProxy InSelectionSet, const FTypedElementSelectionOptions& InSelectionOptions)
 {
 	FSMInstanceManager SMInstance = SMInstanceElementDataUtil::GetSMInstanceFromHandle(InElementHandle);
 	if (SMInstance && Super::DeselectElement(InElementHandle, InSelectionSet, InSelectionOptions))
diff --git a/Engine/Source/Runtime/Engine/Private/UnrealClient.cpp b/Engine/Source/Runtime/Engine/Private/UnrealClient.cpp
--- a/Engine/Source/Runtime/Engine/Private/UnrealClient.cpp
+++ b/Engine/Source/Runtime/Engine/Private/UnrealClient.cpp
@@ -1864,7 +1864,7 @@ FTypedElementHandle FViewport::GetElementHandleAtPoint(int32 X, int32 Y)
 	return FTypedElementHandle();
 }
 
-void FViewport::GetElementHandlesInRect(FIntRect InRect, UTypedElementList* OutElementHandles)
+void FViewport::GetElementHandlesInRect(FIntRect InRect, FTypedElementListRef OutElementHandles)
 {
 	GetElementHandlesInRectImpl(InRect, *OutElementHandles);
 }
diff --git a/Engine/Source/Runtime/Engine/Public/Elements/Framework/EngineElementsLibrary.h b/Engine/Source/Runtime/Engine/Public/Elements/Framework/EngineElementsLibrary.h
--- a/Engine/Source/Runtime/Engine/Public/Elements/Framework/EngineElementsLibrary.h
+++ b/Engine/Source/Runtime/Engine/Public/Elements/Framework/EngineElementsLibrary.h
@@ -21,9 +21,6 @@ struct FSMInstanceElementData;
 struct FSMInstanceId;
 struct FSMInstanceElementId;
 
-class UWorld;
-class UTypedElementList;
-
 UCLASS()
 class ENGINE_API UEngineElementsLibrary : public UBlueprintFunctionLibrary
 {
diff --git a/Engine/Source/Runtime/Engine/Public/Elements/Framework/TypedElementCommonActions.h b/Engine/Source/Runtime/Engine/Public/Elements/Framework/TypedElementCommonActions.h
--- a/Engine/Source/Runtime/Engine/Public/Elements/Framework/TypedElementCommonActions.h
+++ b/Engine/Source/Runtime/Engine/Public/Elements/Framework/TypedElementCommonActions.h
@@ -2,11 +2,11 @@
 
 #pragma once
 
+#include "Elements/Framework/TypedElementListProxy.h"
 #include "Elements/Framework/TypedElementInterfaceCustomization.h"
 #include "Elements/Interfaces/TypedElementWorldInterface.h"
 #include "TypedElementCommonActions.generated.h"
 
-class UTypedElementList;
 class UTypedElementSelectionSet;
 
 /**
@@ -82,7 +82,7 @@ public:
 	 * @note This list should have been pre-normalized via UTypedElementSelectionSet::GetNormalizedSelection or UTypedElementSelectionSet::GetNormalizedElementList.
 	 */
 	UFUNCTION(BlueprintCallable, Category="TypedElementFramework|Common")
-	bool DeleteNormalizedElements(const UTypedElementList* ElementList, UWorld* World, UTypedElementSelectionSet* InSelectionSet, const FTypedElementDeletionOptions& DeletionOptions);
+	bool DeleteNormalizedElements(const FTypedElementListProxy ElementList, UWorld* World, UTypedElementSelectionSet* InSelectionSet, const FTypedElementDeletionOptions& DeletionOptions);
 
 	/**
 	 * Duplicate any elements from the given selection set that can be duplicated.
@@ -96,7 +96,7 @@ public:
 	 * @note This list should have been pre-normalized via UTypedElementSelectionSet::GetNormalizedSelection or UTypedElementSelectionSet::GetNormalizedElementList.
 	 */
 	UFUNCTION(BlueprintCallable, Category="TypedElementFramework|Common")
-	TArray<FTypedElementHandle> DuplicateNormalizedElements(const UTypedElementList* ElementList, UWorld* World, const FVector& LocationOffset);
+	TArray<FTypedElementHandle> DuplicateNormalizedElements(const FTypedElementListProxy ElementList, UWorld* World, const FVector& LocationOffset);
 
 private:
 	/**
diff --git a/Engine/Source/Runtime/Engine/Public/Elements/SMInstance/SMInstanceElementSelectionInterface.h b/Engine/Source/Runtime/Engine/Public/Elements/SMInstance/SMInstanceElementSelectionInterface.h
--- a/Engine/Source/Runtime/Engine/Public/Elements/SMInstance/SMInstanceElementSelectionInterface.h
+++ b/Engine/Source/Runtime/Engine/Public/Elements/SMInstance/SMInstanceElementSelectionInterface.h
@@ -11,6 +11,6 @@ class ENGINE_API USMInstanceElementSelectionInterface : public UTypedElementSele
 	GENERATED_BODY()
 
 public:
-	virtual bool SelectElement(const FTypedElementHandle& InElementHandle, UTypedElementList* InSelectionSet, const FTypedElementSelectionOptions& InSelectionOptions) override;
-	virtual bool DeselectElement(const FTypedElementHandle& InElementHandle, UTypedElementList* InSelectionSet, const FTypedElementSelectionOptions& InSelectionOptions) override;
+	virtual bool SelectElement(const FTypedElementHandle& InElementHandle, FTypedElementListProxy InSelectionSet, const FTypedElementSelectionOptions& InSelectionOptions) override;
+	virtual bool DeselectElement(const FTypedElementHandle& InElementHandle, FTypedElementListProxy InSelectionSet, const FTypedElementSelectionOptions& InSelectionOptions) override;
 };
diff --git a/Engine/Source/Runtime/Engine/Public/UnrealClient.h b/Engine/Source/Runtime/Engine/Public/UnrealClient.h
--- a/Engine/Source/Runtime/Engine/Public/UnrealClient.h
+++ b/Engine/Source/Runtime/Engine/Public/UnrealClient.h
@@ -16,12 +16,12 @@
 #include "RenderResource.h"
 #include "HitProxies.h"
 #include "RenderGraphDefinitions.h"
+#include "Elements/Framework/TypedElementListFwd.h"
 
 class FCanvas;
 class FViewport;
 class FViewportClient;
 class UModel;
-class UTypedElementList;
 
 /**
  * A render target.
@@ -508,7 +508,7 @@ public:
 	 * InRect must be entirely within the viewport's client area.
 	 * If the hit proxies are not cached, this will call ViewportClient->Draw with a hit-testing canvas.
 	 */
-	ENGINE_API void GetElementHandlesInRect(FIntRect InRect, UTypedElementList* OutElementHandles);
+	ENGINE_API void GetElementHandlesInRect(FIntRect InRect, FTypedElementListRef OutElementHandles);
 	ENGINE_API void GetElementHandlesInRect(FIntRect InRect, TSet<FTypedElementHandle>& OutElementHandles);
 
 	/**
diff --git a/Engine/Source/Runtime/TypedElementFramework/Private/Elements/Framework/TypedElementList.cpp b/Engine/Source/Runtime/TypedElementFramework/Private/Elements/Framework/TypedElementList.cpp
--- a/Engine/Source/Runtime/TypedElementFramework/Private/Elements/Framework/TypedElementList.cpp
+++ b/Engine/Source/Runtime/TypedElementFramework/Private/Elements/Framework/TypedElementList.cpp
@@ -21,10 +21,9 @@ void GetElementImpl(const UTypedElementRegistry* InRegistry, const FTypedElement
 } // namespace TypedElementList_Private
 
 
-FTypedElementListLegacySync::FTypedElementListLegacySync(const UTypedElementList* InElementList)
+FTypedElementListLegacySync::FTypedElementListLegacySync(const FTypedElementList& InElementList)
 	: ElementList(InElementList)
 {
-	checkf(ElementList, TEXT("ElementList was null!"));
 }
 
 FTypedElementListLegacySync::FOnSyncEvent& FTypedElementListLegacySync::OnSyncEvent()
@@ -80,8 +79,8 @@ void FTypedElementListLegacySync::ForceBatchOperationDirty()
 }
 
 
-FTypedElementListLegacySyncScopedBatch::FTypedElementListLegacySyncScopedBatch(const UTypedElementList* InElementList, const bool InNotify)
-	: ElementListLegacySync(InElementList->Legacy_GetSyncPtr())
+FTypedElementListLegacySyncScopedBatch::FTypedElementListLegacySyncScopedBatch(const FTypedElementList& InElementList, const bool InNotify)
+	: ElementListLegacySync(InElementList.Legacy_GetSyncPtr())
 	, bNotify(InNotify)
 {
 	if (ElementListLegacySync)
@@ -113,17 +112,14 @@ void FTypedElementListLegacySyncScopedBatch::ForceDirty()
 }
 
 
-UTypedElementList* UTypedElementList::Private_CreateElementList(UTypedElementRegistry* InRegistry)
+FTypedElementListRef FTypedElementList::Private_CreateElementList(UTypedElementRegistry* InRegistry)
 {
-	UTypedElementList* ElementList = NewObject<UTypedElementList>();
-	ElementList->Initialize(InRegistry);
-	return ElementList;
+	return MakeShareable(new FTypedElementList(InRegistry));
 }
 
-void UTypedElementList::Initialize(UTypedElementRegistry* InRegistry)
+FTypedElementList::FTypedElementList(UTypedElementRegistry* InRegistry)
+	: Registry(InRegistry)
 {
-	checkf(!Registry.Get(), TEXT("Initialize has already been called!"));
-	Registry = InRegistry;
 	checkf(InRegistry, TEXT("Registry is null!"));
 
 	ElementCounts.Initialize(InRegistry);
@@ -131,10 +127,8 @@ void UTypedElementList::Initialize(UTypedElementRegistry* InRegistry)
 	InRegistry->Private_OnElementListCreated(this);
 }
 
-void UTypedElementList::BeginDestroy()
+FTypedElementList::~FTypedElementList()
 {
-	Super::BeginDestroy();
-
 	Empty();
 	LegacySync.Reset();
 	if (UTypedElementRegistry* RegistryPtr = Registry.Get())
@@ -144,9 +138,9 @@ void UTypedElementList::BeginDestroy()
 	}
 }
 
-UTypedElementList* UTypedElementList::Clone() const
+FTypedElementListRef FTypedElementList::Clone() const
 {
-	UTypedElementList* ClonedElementList = Private_CreateElementList(Registry.Get());
+	FTypedElementListRef ClonedElementList = Private_CreateElementList(Registry.Get());
 	for (const FTypedElementHandle& ElementHandle : ElementHandles)
 	{
 		ClonedElementList->Add(ElementHandle);
@@ -154,7 +148,7 @@ UTypedElementList* UTypedElementList::Clone() const
 	return ClonedElementList;
 }
 
-UTypedElementInterface* UTypedElementList::GetElementInterface(const FTypedElementHandle& InElementHandle, const TSubclassOf<UTypedElementInterface>& InBaseInterfaceType) const
+UTypedElementInterface* FTypedElementList::GetElementInterface(const FTypedElementHandle& InElementHandle, const TSubclassOf<UTypedElementInterface>& InBaseInterfaceType) const
 {
 	UTypedElementRegistry* RegistryPtr = Registry.Get();
 	return RegistryPtr
@@ -162,7 +156,7 @@ UTypedElementInterface* UTypedElementList::GetElementInterface(const FTypedEleme
 		: nullptr;
 }
 
-bool UTypedElementList::HasElements(const TSubclassOf<UTypedElementInterface>& InBaseInterfaceType) const
+bool FTypedElementList::HasElements(const TSubclassOf<UTypedElementInterface>& InBaseInterfaceType) const
 {
 	bool bHasFilteredElements = false;
 
@@ -182,7 +176,7 @@ bool UTypedElementList::HasElements(const TSubclassOf<UTypedElementInterface>& I
 	return bHasFilteredElements;
 }
 
-int32 UTypedElementList::CountElements(const TSubclassOf<UTypedElementInterface>& InBaseInterfaceType) const
+int32 FTypedElementList::CountElements(const TSubclassOf<UTypedElementInterface>& InBaseInterfaceType) const
 {
 	int32 NumFilteredElements = 0;
 
@@ -202,17 +196,17 @@ int32 UTypedElementList::CountElements(const TSubclassOf<UTypedElementInterface>
 	return NumFilteredElements;
 }
 
-bool UTypedElementList::HasElementsOfType(const FName InElementTypeName) const
+bool FTypedElementList::HasElementsOfType(const FName InElementTypeName) const
 {
 	return CountElementsOfType(InElementTypeName) > 0;
 }
 
-bool UTypedElementList::HasElementsOfType(const FTypedHandleTypeId InElementTypeId) const
+bool FTypedElementList::HasElementsOfType(const FTypedHandleTypeId InElementTypeId) const
 {
 	return CountElementsOfType(InElementTypeId) > 0;
 }
 
-int32 UTypedElementList::CountElementsOfType(const FName InElementTypeName) const
+int32 FTypedElementList::CountElementsOfType(const FName InElementTypeName) const
 {
 	if (UTypedElementRegistry* RegistryPtr = Registry.Get())
 	{
@@ -225,12 +219,12 @@ int32 UTypedElementList::CountElementsOfType(const FName InElementTypeName) cons
 	return 0;
 }
 
-int32 UTypedElementList::CountElementsOfType(const FTypedHandleTypeId InElementTypeId) const
+int32 FTypedElementList::CountElementsOfType(const FTypedHandleTypeId InElementTypeId) const
 {
 	return ElementCounts.GetCounterValue(FTypedElementCounter::GetElementTypeCategoryName(), InElementTypeId);
 }
 
-TArray<FTypedElementHandle> UTypedElementList::GetElementHandles(const TSubclassOf<UTypedElementInterface>& InBaseInterfaceType) const
+TArray<FTypedElementHandle> FTypedElementList::GetElementHandles(const TSubclassOf<UTypedElementInterface>& InBaseInterfaceType) const
 {
 	TArray<FTypedElementHandle> FilteredElementHandles;
 	FilteredElementHandles.Reserve(ElementHandles.Num());
@@ -244,7 +238,7 @@ TArray<FTypedElementHandle> UTypedElementList::GetElementHandles(const TSubclass
 	return FilteredElementHandles;
 }
 
-void UTypedElementList::ForEachElementHandle(TFunctionRef<bool(const FTypedElementHandle&)> InCallback, const TSubclassOf<UTypedElementInterface>& InBaseInterfaceType) const
+void FTypedElementList::ForEachElementHandle(TFunctionRef<bool(const FTypedElementHandle&)> InCallback, const TSubclassOf<UTypedElementInterface>& InBaseInterfaceType) const
 {
 	for (const FTypedElementHandle& ElementHandle : ElementHandles)
 	{
@@ -258,7 +252,7 @@ void UTypedElementList::ForEachElementHandle(TFunctionRef<bool(const FTypedEleme
 	}
 }
 
-bool UTypedElementList::AddElementImpl(FTypedElementHandle&& InElementHandle)
+bool FTypedElementList::AddElementImpl(FTypedElementHandle&& InElementHandle)
 {
 	if (!InElementHandle)
 	{
@@ -280,7 +274,7 @@ bool UTypedElementList::AddElementImpl(FTypedElementHandle&& InElementHandle)
 	return !bAlreadyAdded;
 }
 
-bool UTypedElementList::RemoveElementImpl(const FTypedElementId& InElementId)
+bool FTypedElementList::RemoveElementImpl(const FTypedElementId& InElementId)
 {
 	if (!InElementId)
 	{
@@ -309,13 +303,13 @@ bool UTypedElementList::RemoveElementImpl(const FTypedElementId& InElementId)
 	return bRemoved;
 }
 
-int32 UTypedElementList::RemoveAllElementsImpl(TFunctionRef<bool(const FTypedElementHandle&)> InPredicate)
+int32 FTypedElementList::RemoveAllElementsImpl(TFunctionRef<bool(const FTypedElementHandle&)> InPredicate)
 {
 	int32 RemovedCount = 0;
 
 	if (ElementHandles.Num() > 0)
 	{
-		FTypedElementListLegacySyncScopedBatch LegacySyncBatch(this);
+		FTypedElementListLegacySyncScopedBatch LegacySyncBatch(*this);
 
 		NoteListMayChange();
 
@@ -338,47 +332,47 @@ int32 UTypedElementList::RemoveAllElementsImpl(TFunctionRef<bool(const FTypedEle
 	return RemovedCount;
 }
 
-bool UTypedElementList::ContainsElementImpl(const FTypedElementId& InElementId) const
+bool FTypedElementList::ContainsElementImpl(const FTypedElementId& InElementId) const
 {
 	return InElementId 
 		&& ElementCombinedIds.Contains(InElementId.GetCombinedId());
 }
 
-FTypedElementListLegacySync& UTypedElementList::Legacy_GetSync()
+FTypedElementListLegacySync& FTypedElementList::Legacy_GetSync()
 {
 	if (!LegacySync)
 	{
-		LegacySync = MakeUnique<FTypedElementListLegacySync>(this);
+		LegacySync = MakeUnique<FTypedElementListLegacySync>(*this);
 	}
 	return *LegacySync;
 }
 
-FTypedElementListLegacySync* UTypedElementList::Legacy_GetSyncPtr() const
+FTypedElementListLegacySync* FTypedElementList::Legacy_GetSyncPtr() const
 {
 	return LegacySync.Get();
 }
 
-void UTypedElementList::NotifyPendingChanges()
+void FTypedElementList::NotifyPendingChanges()
 {
 	if (bHasPendingNotify)
 	{
 		bHasPendingNotify = false;
-		OnChangedDelegate.Broadcast(this);
+		OnChangedDelegate.Broadcast(*this);
 		check(!bHasPendingNotify); // This should still be false after emitting the notification!
 	}
 }
 
-void UTypedElementList::ClearPendingChanges()
+void FTypedElementList::ClearPendingChanges()
 {
 	bHasPendingNotify = false;
 }
 
-void UTypedElementList::NoteListMayChange()
+void FTypedElementList::NoteListMayChange()
 {
-	OnPreChangeDelegate.Broadcast(this);
+	OnPreChangeDelegate.Broadcast(*this);
 }
 
-void UTypedElementList::NoteListChanged(const EChangeType InChangeType, const FTypedElementHandle& InElementHandle)
+void FTypedElementList::NoteListChanged(const EChangeType InChangeType, const FTypedElementHandle& InElementHandle)
 {
 	bHasPendingNotify = true;
 
diff --git a/Engine/Source/Runtime/TypedElementFramework/Private/Elements/Framework/TypedElementListLibrary.h b/Engine/Source/Runtime/TypedElementFramework/Private/Elements/Framework/TypedElementListLibrary.h
new file mode 100644
--- /dev/null
+++ b/Engine/Source/Runtime/TypedElementFramework/Private/Elements/Framework/TypedElementListLibrary.h
@@ -0,0 +1,219 @@
+// Copyright Epic Games, Inc. All Rights Reserved.
+
+#pragma once
+
+#include "CoreMinimal.h"
+#include "Elements/Framework/TypedElementRegistry.h"
+#include "Elements/Framework/TypedElementListProxy.h"
+#include "Elements/Framework/TypedElementList.h"
+#include "TypedElementListLibrary.generated.h"
+
+UCLASS()
+class UTypedElementListLibrary : public UObject
+{
+	GENERATED_BODY()
+	
+public:
+	/**
+	 * Create an empty list of elements associated with the given registry.
+	 */
+	UFUNCTION(BlueprintCallable, Category="TypedElementFramework|Registry", meta=(ScriptMethod))
+	static FTypedElementListProxy CreateElementList(UTypedElementRegistry* Registry)
+	{
+		return Registry->CreateElementList();
+	}
+
+	/**
+	 * Clone this list instance.
+	 * @note Only copies elements; does not copy any bindings!
+	 */
+	UFUNCTION(BlueprintCallable, Category="TypedElementFramework|List", meta=(ScriptMethod))
+	static FTypedElementListProxy Clone(const FTypedElementListProxy ElementList)
+	{
+		FTypedElementListConstPtr ElementListPtr = ElementList.GetElementList();
+		return ElementListPtr ? ElementListPtr->Clone() : FTypedElementListProxy();
+	}
+
+	/**
+	 * Get the element handle at the given index.
+	 * @note Use IsValidIndex to test for validity.
+	 */
+	UFUNCTION(BlueprintPure, Category="TypedElementFramework|List", meta=(ScriptMethod))
+	static FTypedElementHandle GetElementHandleAt(const FTypedElementListProxy ElementList, const int32 Index)
+	{
+		FTypedElementListConstPtr ElementListPtr = ElementList.GetElementList();
+		return ElementListPtr ? ElementListPtr->GetElementHandleAt(Index) : FTypedElementHandle();
+	}
+
+	/**
+	 * Get the element interface from the given handle.
+	 */
+	UFUNCTION(BlueprintPure, Category="TypedElementFramework|List", meta=(ScriptMethod))
+	static UTypedElementInterface* GetElementInterface(const FTypedElementListProxy ElementList, const FTypedElementHandle& ElementHandle, const TSubclassOf<UTypedElementInterface>& BaseInterfaceType)
+	{
+		FTypedElementListConstPtr ElementListPtr = ElementList.GetElementList();
+		return ElementListPtr ? ElementListPtr->GetElementInterface(ElementHandle, BaseInterfaceType) : nullptr;
+	}
+
+	/**
+	 * Test whether there are elements in this list, optionally filtering to elements that implement the given interface.
+	 */
+	UFUNCTION(BlueprintCallable, Category="TypedElementFramework|List", meta=(ScriptMethod))
+	static bool HasElements(const FTypedElementListProxy ElementList, const TSubclassOf<UTypedElementInterface>& BaseInterfaceType = nullptr)
+	{
+		FTypedElementListConstPtr ElementListPtr = ElementList.GetElementList();
+		return ElementListPtr && ElementListPtr->HasElements(BaseInterfaceType);
+	}
+
+	/**
+	 * Count the number of elements in this list, optionally filtering to elements that implement the given interface.
+	 */
+	UFUNCTION(BlueprintCallable, Category="TypedElementFramework|List", meta=(ScriptMethod))
+	static int32 CountElements(const FTypedElementListProxy ElementList, const TSubclassOf<UTypedElementInterface>& BaseInterfaceType = nullptr)
+	{
+		FTypedElementListConstPtr ElementListPtr = ElementList.GetElementList();
+		return ElementListPtr ? ElementListPtr->CountElements(BaseInterfaceType) : 0;
+	}
+
+	/**
+	 * Test whether there are elements in this list of the given type.
+	 */
+	UFUNCTION(BlueprintCallable, Category="TypedElementFramework|List", meta=(ScriptMethod))
+	static bool HasElementsOfType(const FTypedElementListProxy ElementList, const FName ElementTypeName)
+	{
+		FTypedElementListConstPtr ElementListPtr = ElementList.GetElementList();
+		return ElementListPtr && ElementListPtr->HasElementsOfType(ElementTypeName);
+	}
+
+	/**
+	 * Count the number of elements in this list of the given type.
+	 */
+	UFUNCTION(BlueprintCallable, Category="TypedElementFramework|List", meta=(ScriptMethod))
+	static int32 CountElementsOfType(const FTypedElementListProxy ElementList, const FName ElementTypeName)
+	{
+		FTypedElementListConstPtr ElementListPtr = ElementList.GetElementList();
+		return ElementListPtr ? ElementListPtr->CountElementsOfType(ElementTypeName) : 0;
+	}
+
+	/**
+	 * Get the handle of every element in this list, optionally filtering to elements that implement the given interface.
+	 */
+	UFUNCTION(BlueprintCallable, Category="TypedElementFramework|List", meta=(ScriptMethod))
+	static TArray<FTypedElementHandle> GetElementHandles(const FTypedElementListProxy ElementList, const TSubclassOf<UTypedElementInterface>& BaseInterfaceType = nullptr)
+	{
+		FTypedElementListConstPtr ElementListPtr = ElementList.GetElementList();
+		return ElementListPtr ? ElementListPtr->GetElementHandles(BaseInterfaceType) : TArray<FTypedElementHandle>();
+	}
+
+	/**
+	 * Is the given index a valid entry within this element list?
+	 */
+	UFUNCTION(BlueprintPure, Category="TypedElementFramework|List", meta=(ScriptMethod))
+	static bool IsValidIndex(const FTypedElementListProxy ElementList, const int32 Index)
+	{
+		FTypedElementListConstPtr ElementListPtr = ElementList.GetElementList();
+		return ElementListPtr && ElementListPtr->IsValidIndex(Index);
+	}
+
+	/**
+	 * Get the number of entries within this element list.
+	 */
+	UFUNCTION(BlueprintPure, Category="TypedElementFramework|List", meta=(ScriptMethod))
+	static int32 Num(const FTypedElementListProxy ElementList)
+	{
+		FTypedElementListConstPtr ElementListPtr = ElementList.GetElementList();
+		return ElementListPtr ? ElementListPtr->Num() : 0;
+	}
+
+	/**
+	 * Shrink this element list storage to avoid slack.
+	 */
+	UFUNCTION(BlueprintCallable, Category="TypedElementFramework|List", meta=(ScriptMethod))
+	static void Shrink(FTypedElementListProxy ElementList)
+	{
+		if (FTypedElementListPtr ElementListPtr = ElementList.GetElementList())
+		{
+			ElementListPtr->Shrink();
+		}
+	}
+
+	/**
+	 * Pre-allocate enough memory in this element list to store the given number of entries.
+	 */
+	UFUNCTION(BlueprintCallable, Category="TypedElementFramework|List", meta=(ScriptMethod))
+	static void Reserve(FTypedElementListProxy ElementList, const int32 Size)
+	{
+		if (FTypedElementListPtr ElementListPtr = ElementList.GetElementList())
+		{
+			ElementListPtr->Reserve(Size);
+		}
+	}
+
+	/**
+	 * Remove all entries from this element list, potentially leaving space allocated for the given number of entries.
+	 */
+	UFUNCTION(BlueprintCallable, Category="TypedElementFramework|List", meta=(ScriptMethod))
+	static void Empty(FTypedElementListProxy ElementList, const int32 Slack = 0)
+	{
+		if (FTypedElementListPtr ElementListPtr = ElementList.GetElementList())
+		{
+			ElementListPtr->Empty(Slack);
+		}
+	}
+
+	/**
+	 * Remove all entries from this element list, preserving existing allocations.
+	 */
+	UFUNCTION(BlueprintCallable, Category="TypedElementFramework|List", meta=(ScriptMethod))
+	static void Reset(FTypedElementListProxy ElementList)
+	{
+		if (FTypedElementListPtr ElementListPtr = ElementList.GetElementList())
+		{
+			ElementListPtr->Reset();
+		}
+	}
+
+	/**
+	 * Does this element list contain an entry for the given element handle?
+	 */
+	UFUNCTION(BlueprintPure, Category="TypedElementFramework|List", meta=(ScriptMethod))
+	static bool Contains(const FTypedElementListProxy ElementList, const FTypedElementHandle& ElementHandle)
+	{
+		FTypedElementListConstPtr ElementListPtr = ElementList.GetElementList();
+		return ElementListPtr && ElementListPtr->Contains(ElementHandle);
+	}
+
+	/**
+	 * Add the given element handle to this element list, if it isn't already in the list.
+	 * @return True if the element handle was added, false if it is already in the list.
+	 */
+	UFUNCTION(BlueprintCallable, Category="TypedElementFramework|List", meta=(ScriptMethod))
+	static bool Add(FTypedElementListProxy ElementList, const FTypedElementHandle& ElementHandle)
+	{
+		FTypedElementListPtr ElementListPtr = ElementList.GetElementList();
+		return ElementListPtr && ElementListPtr->Add(ElementHandle);
+	}
+
+	/**
+	 * Append the given element handles to this element list, for any that already in the list.
+	 */
+	UFUNCTION(BlueprintCallable, Category="TypedElementFramework|List", meta=(ScriptMethod))
+	static void Append(FTypedElementListProxy ElementList, const TArray<FTypedElementHandle>& ElementHandles)
+	{
+		if (FTypedElementListPtr ElementListPtr = ElementList.GetElementList())
+		{
+			ElementListPtr->Append(ElementHandles);
+		}
+	}
+
+	/**
+	 * Remove the given element handle from this element list, if it is in the list.
+	 * @return True if the element handle was removed, false if it isn't in the list.
+	 */
+	UFUNCTION(BlueprintCallable, Category="TypedElementFramework|List", meta=(ScriptMethod))
+	static bool Remove(FTypedElementListProxy ElementList, const FTypedElementHandle& ElementHandle)
+	{
+		FTypedElementListPtr ElementListPtr = ElementList.GetElementList();
+		return ElementListPtr && ElementListPtr->Remove(ElementHandle);
+	}
+};
diff --git a/Engine/Source/Runtime/TypedElementFramework/Private/Elements/Framework/TypedElementRegistry.cpp b/Engine/Source/Runtime/TypedElementFramework/Private/Elements/Framework/TypedElementRegistry.cpp
--- a/Engine/Source/Runtime/TypedElementFramework/Private/Elements/Framework/TypedElementRegistry.cpp
+++ b/Engine/Source/Runtime/TypedElementFramework/Private/Elements/Framework/TypedElementRegistry.cpp
@@ -187,9 +187,9 @@ FTypedElementHandle UTypedElementRegistry::GetElementHandle(const FTypedElementI
 	return ElementHandle;
 }
 
-UTypedElementList* UTypedElementRegistry::CreateElementList(TArrayView<const FTypedElementId> InElementIds)
+FTypedElementListRef UTypedElementRegistry::CreateElementList(TArrayView<const FTypedElementId> InElementIds)
 {
-	UTypedElementList* ElementList = CreateElementList();
+	FTypedElementListRef ElementList = CreateElementList();
 
 	for (const FTypedElementId& ElementId : InElementIds)
 	{
@@ -202,9 +202,9 @@ UTypedElementList* UTypedElementRegistry::CreateElementList(TArrayView<const FTy
 	return ElementList;
 }
 
-UTypedElementList* UTypedElementRegistry::CreateElementList(TArrayView<const FTypedElementHandle> InElementHandles)
+FTypedElementListRef UTypedElementRegistry::CreateElementList(TArrayView<const FTypedElementHandle> InElementHandles)
 {
-	UTypedElementList* ElementList = CreateElementList();
+	FTypedElementListRef ElementList = CreateElementList();
 	ElementList->Append(InElementHandles);
 	return ElementList;
 }
@@ -213,7 +213,7 @@ void UTypedElementRegistry::NotifyElementListPendingChanges()
 {
 	FReadScopeLock ActiveElementListsLock(ActiveElementListsRW);
 
-	for (UTypedElementList* ActiveElementList : ActiveElementLists)
+	for (FTypedElementList* ActiveElementList : ActiveElementLists)
 	{
 		ActiveElementList->NotifyPendingChanges();
 	}
diff --git a/Engine/Source/Runtime/TypedElementFramework/Private/Elements/Framework/TypedElementUtil.cpp b/Engine/Source/Runtime/TypedElementFramework/Private/Elements/Framework/TypedElementUtil.cpp
--- a/Engine/Source/Runtime/TypedElementFramework/Private/Elements/Framework/TypedElementUtil.cpp
+++ b/Engine/Source/Runtime/TypedElementFramework/Private/Elements/Framework/TypedElementUtil.cpp
@@ -6,7 +6,7 @@
 namespace TypedElementUtil
 {
 
-void BatchElementsByType(const UTypedElementList* InElementsToBatch, TMap<FTypedHandleTypeId, TArray<FTypedElementHandle>>& OutElementsByType)
+void BatchElementsByType(FTypedElementListConstRef InElementsToBatch, TMap<FTypedHandleTypeId, TArray<FTypedElementHandle>>& OutElementsByType)
 {
 	OutElementsByType.Reset();
 	InElementsToBatch->ForEachElementHandle([&OutElementsByType](const FTypedElementHandle& InElementHandle)
diff --git a/Engine/Source/Runtime/TypedElementFramework/Private/TypedElementFrameworkTests.cpp b/Engine/Source/Runtime/TypedElementFramework/Private/TypedElementFrameworkTests.cpp
--- a/Engine/Source/Runtime/TypedElementFramework/Private/TypedElementFrameworkTests.cpp
+++ b/Engine/Source/Runtime/TypedElementFramework/Private/TypedElementFrameworkTests.cpp
@@ -109,7 +109,7 @@ bool FTypedElementRegistrySmokeTest::RunTest(const FString& Parameters)
 	TestInterfaceAccess(TypedElement1.AcquireHandle());
 	TestInterfaceAccess(UntypedElement1.AcquireHandle());
 
-	UTypedElementList* ElementList = Registry->CreateElementList();
+	FTypedElementListPtr ElementList = Registry->CreateElementList();
 	ElementList->Add(TypedElement1);
 	ElementList->Add(TypedElement2);
 	ElementList->Add(TypedElement3);
@@ -129,8 +129,7 @@ bool FTypedElementRegistrySmokeTest::RunTest(const FString& Parameters)
 		return true;
 	});
 
-	ElementList->Empty();
-	ElementList = nullptr;
+	ElementList.Reset();
 
 	Registry->DestroyElement(TypedElement1);
 	Registry->DestroyElement(TypedElement2);
@@ -140,6 +139,9 @@ bool FTypedElementRegistrySmokeTest::RunTest(const FString& Parameters)
 	Registry->DestroyElement(UntypedElement2);
 	Registry->DestroyElement(UntypedElement3);
 
+	// Verify that there were no leaks
+	Registry->ProcessDeferredElementsToDestroy();
+
 	return true;
 }
 
@@ -160,7 +162,7 @@ bool FTypedElementRegistryPerfTest::RunTest(const FString& Parameters)
 
 	TArray<TTypedElementOwner<FTestTypedElementData>> TypedOwnerHandles;
 	TArray<FTypedElementOwner> UntypedOwnerHandles;
-	UTypedElementList* ElementList = Registry->CreateElementList();
+	FTypedElementListPtr ElementList = Registry->CreateElementList();
 
 	// Create typed handles
 	{
@@ -226,8 +228,7 @@ bool FTypedElementRegistryPerfTest::RunTest(const FString& Parameters)
 	// Clear the element list
 	{
 		FScopedDurationTimeLogger Timer(FString::Printf(TEXT("Reset %d elements in list"), ElementList->Num()));
-		ElementList->Empty();
-		ElementList = nullptr;
+		ElementList.Reset();
 	}
 
 	// Destroy typed handles
@@ -254,6 +255,9 @@ bool FTypedElementRegistryPerfTest::RunTest(const FString& Parameters)
 		UntypedOwnerHandles.Reset();
 	}
 
+	// Verify that there were no leaks
+	Registry->ProcessDeferredElementsToDestroy();
+
 	return true;
 }
 
diff --git a/Engine/Source/Runtime/TypedElementFramework/Public/Elements/Framework/TypedElementId.h b/Engine/Source/Runtime/TypedElementFramework/Public/Elements/Framework/TypedElementId.h
--- a/Engine/Source/Runtime/TypedElementFramework/Public/Elements/Framework/TypedElementId.h
+++ b/Engine/Source/Runtime/TypedElementFramework/Public/Elements/Framework/TypedElementId.h
@@ -76,7 +76,7 @@ public:
 	/**
 	 * Access the combined value of this element ID.
 	 * @note You typically don't want to store this directly as the element ID could be re-used.
-	 *       It is primarily useful as a secondary cache where something is keeping a reference to an element ID or element handle (eg, how UTypedElementList uses it internally).
+	 *       It is primarily useful as a secondary cache where something is keeping a reference to an element ID or element handle (eg, how FTypedElementList uses it internally).
 	 */
 	FORCEINLINE FTypedHandleCombinedId GetCombinedId() const
 	{
diff --git a/Engine/Source/Runtime/TypedElementFramework/Public/Elements/Framework/TypedElementList.h b/Engine/Source/Runtime/TypedElementFramework/Public/Elements/Framework/TypedElementList.h
--- a/Engine/Source/Runtime/TypedElementFramework/Public/Elements/Framework/TypedElementList.h
+++ b/Engine/Source/Runtime/TypedElementFramework/Public/Elements/Framework/TypedElementList.h
@@ -10,9 +10,8 @@
 #include "Templates/UniquePtr.h"
 #include "Elements/Framework/TypedElementHandle.h"
 #include "Elements/Framework/TypedElementCounter.h"
-#include "TypedElementList.generated.h"
+#include "Elements/Framework/TypedElementListFwd.h"
 
-class UTypedElementList;
 class UTypedElementRegistry;
 
 namespace TypedElementList_Private
@@ -79,11 +78,11 @@ public:
 		BatchComplete,
 	};
 	
-	FTypedElementListLegacySync(const UTypedElementList* InElementList);
+	FTypedElementListLegacySync(const FTypedElementList& InElementList);
 
 	void Private_EmitSyncEvent(const ESyncType InSyncType, const FTypedElementHandle& InElementHandle = FTypedElementHandle());
 
-	DECLARE_EVENT_FourParams(FTypedElementListLegacySync, FOnSyncEvent, const UTypedElementList* /*InElementList*/, ESyncType /*InSyncType*/, const FTypedElementHandle& /*InElementHandle*/, bool /*bIsWithinBatchOperation*/);
+	DECLARE_EVENT_FourParams(FTypedElementListLegacySync, FOnSyncEvent, const FTypedElementList& /*InElementList*/, ESyncType /*InSyncType*/, const FTypedElementHandle& /*InElementHandle*/, bool /*bIsWithinBatchOperation*/);
 	FOnSyncEvent& OnSyncEvent();
 
 	bool IsRunningBatchOperation() const;
@@ -93,7 +92,7 @@ public:
 	void ForceBatchOperationDirty();
 
 private:
-	const UTypedElementList* ElementList;
+	const FTypedElementList& ElementList;
 
 	FOnSyncEvent OnSyncEventDelegate;
 
@@ -108,7 +107,7 @@ private:
 class TYPEDELEMENTFRAMEWORK_API FTypedElementListLegacySyncScopedBatch
 {
 public:
-	explicit FTypedElementListLegacySyncScopedBatch(const UTypedElementList* InElementList, const bool InNotify = true);
+	explicit FTypedElementListLegacySyncScopedBatch(const FTypedElementList& InElementList, const bool InNotify = true);
 	~FTypedElementListLegacySyncScopedBatch();
 
 	FTypedElementListLegacySyncScopedBatch(const FTypedElementListLegacySyncScopedBatch&) = delete;
@@ -129,26 +128,21 @@ private:
  * A list of element handles.
  * Provides high-level access to groups of elements, including accessing elements that implement specific interfaces.
  */
-UCLASS(Transient)
-class TYPEDELEMENTFRAMEWORK_API UTypedElementList : public UObject
+class TYPEDELEMENTFRAMEWORK_API FTypedElementList final : public TSharedFromThis<FTypedElementList>
 {
-	GENERATED_BODY()
-
 public:
-	//~ UObject interface
-	virtual void BeginDestroy() override;
+	~FTypedElementList();
 
 	/**
 	 * Internal function used by the element registry to create an element list instance.
 	 */
-	static UTypedElementList* Private_CreateElementList(UTypedElementRegistry* InRegistry);
+	static FTypedElementListRef Private_CreateElementList(UTypedElementRegistry* InRegistry);
 
 	/**
 	 * Clone this list instance.
 	 * @note Only copies elements; does not copy any bindings!
 	 */
-	UFUNCTION(BlueprintCallable, BlueprintPure=false, Category="TypedElementFramework|List")
-	UTypedElementList* Clone() const;
+	FTypedElementListRef Clone() const;
 
 	/**
 	 * Get the element handle at the given index.
@@ -163,7 +157,6 @@ public:
 	 * Get the element handle at the given index.
 	 * @note Use IsValidIndex to test for validity.
 	 */
-	UFUNCTION(BlueprintPure, Category="TypedElementFramework|List")
 	FORCEINLINE FTypedElementHandle GetElementHandleAt(const int32 InIndex) const
 	{
 		return ElementHandles[InIndex];
@@ -255,39 +248,33 @@ public:
 	/**
 	 * Get the element interface from the given handle.
 	 */
-	UFUNCTION(BlueprintPure, Category="TypedElementFramework|List")
 	UTypedElementInterface* GetElementInterface(const FTypedElementHandle& InElementHandle, const TSubclassOf<UTypedElementInterface>& InBaseInterfaceType) const;
 
 	/**
 	 * Test whether there are elements in this list, optionally filtering to elements that implement the given interface.
 	 */
-	UFUNCTION(BlueprintCallable, BlueprintPure=false, Category="TypedElementFramework|List")
 	bool HasElements(const TSubclassOf<UTypedElementInterface>& InBaseInterfaceType = nullptr) const;
 
 	/**
 	 * Count the number of elements in this list, optionally filtering to elements that implement the given interface.
 	 */
-	UFUNCTION(BlueprintCallable, BlueprintPure=false, Category="TypedElementFramework|List")
 	int32 CountElements(const TSubclassOf<UTypedElementInterface>& InBaseInterfaceType = nullptr) const;
 
 	/**
 	 * Test whether there are elements in this list of the given type.
 	 */
-	UFUNCTION(BlueprintCallable, BlueprintPure=false, Category="TypedElementFramework|List")
 	bool HasElementsOfType(const FName InElementTypeName) const;
 	bool HasElementsOfType(const FTypedHandleTypeId InElementTypeId) const;
 
 	/**
 	 * Count the number of elements in this list of the given type.
 	 */
-	UFUNCTION(BlueprintCallable, BlueprintPure=false, Category="TypedElementFramework|List")
 	int32 CountElementsOfType(const FName InElementTypeName) const;
 	int32 CountElementsOfType(const FTypedHandleTypeId InElementTypeId) const;
 
 	/**
 	 * Get the handle of every element in this list, optionally filtering to elements that implement the given interface.
 	 */
-	UFUNCTION(BlueprintCallable, BlueprintPure=false, Category="TypedElementFramework|List")
 	TArray<FTypedElementHandle> GetElementHandles(const TSubclassOf<UTypedElementInterface>& InBaseInterfaceType = nullptr) const;
 
 	/**
@@ -332,7 +319,6 @@ public:
 	/**
 	 * Is the given index a valid entry within this element list?
 	 */
-	UFUNCTION(BlueprintPure, Category="TypedElementFramework|List")
 	FORCEINLINE bool IsValidIndex(const int32 InIndex) const
 	{
 		return ElementHandles.IsValidIndex(InIndex);
@@ -341,7 +327,6 @@ public:
 	/**
 	 * Get the number of entries within this element list.
 	 */
-	UFUNCTION(BlueprintPure, Category="TypedElementFramework|List")
 	FORCEINLINE int32 Num() const
 	{
 		return ElementHandles.Num();
@@ -350,7 +335,6 @@ public:
 	/**
 	 * Shrink this element list storage to avoid slack.
 	 */
-	UFUNCTION(BlueprintCallable, Category="TypedElementFramework|List")
 	FORCEINLINE void Shrink()
 	{
 		ElementCombinedIds.Shrink();
@@ -360,7 +344,6 @@ public:
 	/**
 	 * Pre-allocate enough memory in this element list to store the given number of entries.
 	 */
-	UFUNCTION(BlueprintCallable, Category="TypedElementFramework|List")
 	FORCEINLINE void Reserve(const int32 InSize)
 	{
 		ElementCombinedIds.Reserve(InSize);
@@ -370,7 +353,6 @@ public:
 	/**
 	 * Remove all entries from this element list, potentially leaving space allocated for the given number of entries.
 	 */
-	UFUNCTION(BlueprintCallable, Category="TypedElementFramework|List")
 	FORCEINLINE void Empty(const int32 InSlack = 0)
 	{
 		NoteListMayChange();
@@ -383,7 +365,6 @@ public:
 	/**
 	 * Remove all entries from this element list, preserving existing allocations.
 	 */
-	UFUNCTION(BlueprintCallable, Category="TypedElementFramework|List")
 	FORCEINLINE void Reset()
 	{
 		NoteListMayChange();
@@ -404,7 +385,6 @@ public:
 	/**
 	 * Does this element list contain an entry for the given element handle?
 	 */
-	UFUNCTION(BlueprintPure, Category="TypedElementFramework|List")
 	FORCEINLINE bool Contains(const FTypedElementHandle& InElementHandle) const
 	{
 		return ContainsElementImpl(InElementHandle.GetId());
@@ -423,7 +403,6 @@ public:
 	 * Add the given element handle to this element list, if it isn't already in the list.
 	 * @return True if the element handle was added, false if it is already in the list.
 	 */
-	UFUNCTION(BlueprintCallable, Category="TypedElementFramework|List")
 	FORCEINLINE bool Add(const FTypedElementHandle& InElementHandle)
 	{
 		return AddElementImpl(CopyTemp(InElementHandle));
@@ -448,21 +427,12 @@ public:
 		return AddElementImpl(InElementOwner.AcquireHandle());
 	}
 
-	/**
-	 * Append the given element handles to this element list, for any that already in the list.
-	 */
-	UFUNCTION(BlueprintCallable, Category="TypedElementFramework|List")
-	FORCEINLINE void Append(const TArray<FTypedElementHandle>& InElementHandles)
-	{
-		Append(MakeArrayView(InElementHandles));
-	}
-	
 	/**
 	 * Append the given element handles to this element list, for any that already in the list.
 	 */
 	void Append(TArrayView<const FTypedElementHandle> InElementHandles)
 	{
-		FTypedElementListLegacySyncScopedBatch LegacySyncBatch(this);
+		FTypedElementListLegacySyncScopedBatch LegacySyncBatch(*this);
 
 		Reserve(Num() + InElementHandles.Num());
 		for (const FTypedElementHandle& ElementHandle : InElementHandles)
@@ -486,7 +456,7 @@ public:
 	template <typename ElementDataType>
 	void Append(TArrayView<const TTypedElementOwner<ElementDataType>> InElementOwners)
 	{
-		FTypedElementListLegacySyncScopedBatch LegacySyncBatch(this);
+		FTypedElementListLegacySyncScopedBatch LegacySyncBatch(*this);
 
 		Reserve(Num() + InElementOwners.Num());
 		for (const TTypedElementOwner<ElementDataType>& ElementOwner : InElementOwners)
@@ -508,7 +478,6 @@ public:
 	 * Remove the given element handle from this element list, if it is in the list.
 	 * @return True if the element handle was removed, false if it isn't in the list.
 	 */
-	UFUNCTION(BlueprintCallable, Category="TypedElementFramework|List")
 	FORCEINLINE bool Remove(const FTypedElementHandle& InElementHandle)
 	{
 		return RemoveElementImpl(InElementHandle.GetId());
@@ -560,7 +529,7 @@ public:
 	 * Access the delegate that is invoked whenever this element list is potentially about to change.
 	 * @note This may be called even if no actual change happens, so may be called multiple times without a corresponding OnChanged notification.
 	 */
-	DECLARE_EVENT_OneParam(UTypedElementList, FOnPreChange, const UTypedElementList* /*InElementList*/);
+	DECLARE_EVENT_OneParam(FTypedElementList, FOnPreChange, const FTypedElementList& /*InElementList*/);
 	FOnPreChange& OnPreChange()
 	{
 		return OnPreChangeDelegate;
@@ -570,7 +539,7 @@ public:
 	 * Access the delegate that is invoked whenever this element list has been changed.
 	 * @note This is called automatically at the end of each frame, but can also be manually invoked by NotifyPendingChanges.
 	 */
-	DECLARE_EVENT_OneParam(UTypedElementList, FOnChanged, const UTypedElementList* /*InElementList*/);
+	DECLARE_EVENT_OneParam(FTypedElementList, FOnChanged, const FTypedElementList& /*InElementList*/);
 	FOnChanged& OnChanged()
 	{
 		return OnChangedDelegate;
@@ -620,7 +589,7 @@ private:
 		Cleared,
 	};
 
-	void Initialize(UTypedElementRegistry* InRegistry);
+	explicit FTypedElementList(UTypedElementRegistry* InRegistry);
 
 	bool AddElementImpl(FTypedElementHandle&& InElementHandle);
 	bool RemoveElementImpl(const FTypedElementId& InElementId);
diff --git a/Engine/Source/Runtime/TypedElementFramework/Public/Elements/Framework/TypedElementListFwd.h b/Engine/Source/Runtime/TypedElementFramework/Public/Elements/Framework/TypedElementListFwd.h
new file mode 100644
--- /dev/null
+++ b/Engine/Source/Runtime/TypedElementFramework/Public/Elements/Framework/TypedElementListFwd.h
@@ -0,0 +1,13 @@
+// Copyright Epic Games, Inc. All Rights Reserved.
+
+#pragma once
+
+#include "Templates/SharedPointer.h"
+
+class FTypedElementList;
+
+using FTypedElementListPtr = TSharedPtr<FTypedElementList>;
+using FTypedElementListRef = TSharedRef<FTypedElementList>;
+
+using FTypedElementListConstPtr = TSharedPtr<const FTypedElementList>;
+using FTypedElementListConstRef = TSharedRef<const FTypedElementList>;
diff --git a/Engine/Source/Runtime/TypedElementFramework/Public/Elements/Framework/TypedElementListProxy.h b/Engine/Source/Runtime/TypedElementFramework/Public/Elements/Framework/TypedElementListProxy.h
new file mode 100644
--- /dev/null
+++ b/Engine/Source/Runtime/TypedElementFramework/Public/Elements/Framework/TypedElementListProxy.h
@@ -0,0 +1,43 @@
+// Copyright Epic Games, Inc. All Rights Reserved.
+
+#pragma once
+
+#include "CoreMinimal.h"
+#include "Elements/Framework/TypedElementListFwd.h"
+#include "TypedElementListProxy.generated.h"
+
+/**
+ * A list of element handles (proxy to a FTypedElementList instance).
+ * Provides high-level access to groups of elements, including accessing elements that implement specific interfaces.
+ */
+USTRUCT(BlueprintType, DisplayName="TypedElementList", meta=(ScriptName="TypedElementList"))
+struct TYPEDELEMENTFRAMEWORK_API FTypedElementListProxy
+{
+	GENERATED_BODY()
+	
+public:
+	FTypedElementListProxy() = default;
+
+	FTypedElementListProxy(FTypedElementListRef InElementList)
+		: ElementList(MoveTemp(InElementList))
+	{
+	}
+
+	FTypedElementListProxy(FTypedElementListConstRef InElementList)
+		: ElementList(ConstCastSharedRef<FTypedElementList>(InElementList))
+	{
+	}
+
+	FTypedElementListPtr GetElementList()
+	{
+		return ElementList;
+	}
+
+	FTypedElementListConstPtr GetElementList() const
+	{
+		return ElementList;
+	}
+	
+private:
+	FTypedElementListPtr ElementList;
+};
diff --git a/Engine/Source/Runtime/TypedElementFramework/Public/Elements/Framework/TypedElementRegistry.h b/Engine/Source/Runtime/TypedElementFramework/Public/Elements/Framework/TypedElementRegistry.h
--- a/Engine/Source/Runtime/TypedElementFramework/Public/Elements/Framework/TypedElementRegistry.h
+++ b/Engine/Source/Runtime/TypedElementFramework/Public/Elements/Framework/TypedElementRegistry.h
@@ -242,27 +242,26 @@ public:
 	/**
 	 * Create an empty list of elements associated with this registry.
 	 */
-	UFUNCTION(BlueprintCallable, Category="TypedElementFramework|Registry")
-	FORCEINLINE UTypedElementList* CreateElementList()
+	FORCEINLINE FTypedElementListRef CreateElementList()
 	{
-		return UTypedElementList::Private_CreateElementList(this);
+		return FTypedElementList::Private_CreateElementList(this);
 	}
 
 	/**
 	 * Create an empty list of elements associated with this registry, populated from the given minimal IDs that are valid.
 	 */
-	UTypedElementList* CreateElementList(TArrayView<const FTypedElementId> InElementIds);
+	FTypedElementListRef CreateElementList(TArrayView<const FTypedElementId> InElementIds);
 
 	/**
 	 * Create an empty list of elements associated with this registry, populated from the given handles that are valid.
 	 */
-	UTypedElementList* CreateElementList(TArrayView<const FTypedElementHandle> InElementHandles);
+	FTypedElementListRef CreateElementList(TArrayView<const FTypedElementHandle> InElementHandles);
 
 	/**
 	 * Create an empty list of elements associated with this registry, populated from the given owners that are valid.
 	 */
 	template <typename ElementDataType>
-	FORCEINLINE UTypedElementList* CreateElementList(const TArray<TTypedElementOwner<ElementDataType>>& InElementOwners)
+	FORCEINLINE FTypedElementListRef CreateElementList(const TArray<TTypedElementOwner<ElementDataType>>& InElementOwners)
 	{
 		return CreateElementList(MakeArrayView(InElementOwners));
 	}
@@ -271,26 +270,26 @@ public:
 	 * Create an empty list of elements associated with this registry, populated from the given owners that are valid.
 	 */
 	template <typename ElementDataType>
-	UTypedElementList* CreateElementList(TArrayView<const TTypedElementOwner<ElementDataType>> InElementOwners)
+	FTypedElementListRef CreateElementList(TArrayView<const TTypedElementOwner<ElementDataType>> InElementOwners)
 	{
-		UTypedElementList* ElementList = CreateElementList();
+		FTypedElementListRef ElementList = CreateElementList();
 		ElementList->Append(InElementOwners);
 		return ElementList;
 	}
 
-	void Private_OnElementListCreated(UTypedElementList* InElementList)
+	void Private_OnElementListCreated(FTypedElementList* InElementList)
 	{
 		FWriteScopeLock ActiveElementListsLock(ActiveElementListsRW);
 		ActiveElementLists.Add(InElementList);
 	}
 
-	void Private_OnElementListDestroyed(UTypedElementList* InElementList)
+	void Private_OnElementListDestroyed(FTypedElementList* InElementList)
 	{
 		FWriteScopeLock ActiveElementListsLock(ActiveElementListsRW);
 		ActiveElementLists.Remove(InElementList);
 	}
 
-	// Note: Access for UTypedElementList
+	// Note: Access for FTypedElementList
 	FORCEINLINE void Private_GetElementImpl(const FTypedElementHandle& InElementHandle, const UClass* InBaseInterfaceType, FTypedElement& OutElement) const
 	{
 		GetElementImpl(InElementHandle, InBaseInterfaceType, OutElement);
@@ -519,7 +518,7 @@ private:
 	TSortedMap<FName, FTypedHandleTypeId, FDefaultAllocator, FNameFastLess> RegisteredElementTypesNameToId;
 
 	mutable FRWLock ActiveElementListsRW;
-	TSet<UTypedElementList*> ActiveElementLists;
+	TSet<FTypedElementList*> ActiveElementLists;
 
 	uint8 DisableElementDestructionOnGCCount = 0;
 	bool bIsWithinFrame = false;
diff --git a/Engine/Source/Runtime/TypedElementFramework/Public/Elements/Framework/TypedElementUtil.h b/Engine/Source/Runtime/TypedElementFramework/Public/Elements/Framework/TypedElementUtil.h
--- a/Engine/Source/Runtime/TypedElementFramework/Public/Elements/Framework/TypedElementUtil.h
+++ b/Engine/Source/Runtime/TypedElementFramework/Public/Elements/Framework/TypedElementUtil.h
@@ -3,8 +3,7 @@
 #pragma once
 
 #include "Elements/Framework/TypedElementHandle.h"
-
-class UTypedElementList;
+#include "Elements/Framework/TypedElementListFwd.h"
 
 namespace TypedElementUtil
 {
@@ -12,7 +11,7 @@ namespace TypedElementUtil
 /**
  * Batch the given elements by their type.
  */
-TYPEDELEMENTFRAMEWORK_API void BatchElementsByType(const UTypedElementList* InElementsToBatch, TMap<FTypedHandleTypeId, TArray<FTypedElementHandle>>& OutElementsByType);
+TYPEDELEMENTFRAMEWORK_API void BatchElementsByType(FTypedElementListConstRef InElementsToBatch, TMap<FTypedHandleTypeId, TArray<FTypedElementHandle>>& OutElementsByType);
 TYPEDELEMENTFRAMEWORK_API void BatchElementsByType(TArrayView<const FTypedElementHandle> InElementsToBatch, TMap<FTypedHandleTypeId, TArray<FTypedElementHandle>>& OutElementsByType);
 
 } // namespace TypedElementUtil
diff --git a/Engine/Source/Runtime/TypedElementRuntime/Private/Elements/Framework/TypedElementListObjectUtil.cpp b/Engine/Source/Runtime/TypedElementRuntime/Private/Elements/Framework/TypedElementListObjectUtil.cpp
--- a/Engine/Source/Runtime/TypedElementRuntime/Private/Elements/Framework/TypedElementListObjectUtil.cpp
+++ b/Engine/Source/Runtime/TypedElementRuntime/Private/Elements/Framework/TypedElementListObjectUtil.cpp
@@ -14,7 +14,7 @@ UObject* GetObjectOfType(const TTypedElement<UTypedElementObjectInterface>& InOb
 		: nullptr;
 }
 
-bool HasObjects(const UTypedElementList* InElementList, const UClass* InRequiredClass)
+bool HasObjects(FTypedElementListConstRef InElementList, const UClass* InRequiredClass)
 {
 	bool bHasObjects = false;
 
@@ -27,7 +27,7 @@ bool HasObjects(const UTypedElementList* InElementList, const UClass* InRequired
 	return bHasObjects;
 }
 
-int32 CountObjects(const UTypedElementList* InElementList, const UClass* InRequiredClass)
+int32 CountObjects(FTypedElementListConstRef InElementList, const UClass* InRequiredClass)
 {
 	int32 NumObjects = 0;
 
@@ -40,7 +40,7 @@ int32 CountObjects(const UTypedElementList* InElementList, const UClass* InRequi
 	return NumObjects;
 }
 
-void ForEachObject(const UTypedElementList* InElementList, TFunctionRef<bool(UObject*)> InCallback, const UClass* InRequiredClass)
+void ForEachObject(FTypedElementListConstRef InElementList, TFunctionRef<bool(UObject*)> InCallback, const UClass* InRequiredClass)
 {
 	InElementList->ForEachElement<UTypedElementObjectInterface>([&InCallback, InRequiredClass](const TTypedElement<UTypedElementObjectInterface>& InObjectElement)
 	{
@@ -52,7 +52,7 @@ void ForEachObject(const UTypedElementList* InElementList, TFunctionRef<bool(UOb
 	});
 }
 
-TArray<UObject*> GetObjects(const UTypedElementList* InElementList, const UClass* InRequiredClass)
+TArray<UObject*> GetObjects(FTypedElementListConstRef InElementList, const UClass* InRequiredClass)
 {
 	TArray<UObject*> SelectedObjects;
 	SelectedObjects.Reserve(InElementList->Num());
@@ -66,7 +66,7 @@ TArray<UObject*> GetObjects(const UTypedElementList* InElementList, const UClass
 	return SelectedObjects;
 }
 
-UObject* GetTopObject(const UTypedElementList* InElementList, const UClass* InRequiredClass)
+UObject* GetTopObject(FTypedElementListConstRef InElementList, const UClass* InRequiredClass)
 {
 	TTypedElement<UTypedElementObjectInterface> TempElement;
 	for (int32 ElementIndex = 0; ElementIndex < InElementList->Num(); ++ElementIndex)
@@ -82,7 +82,7 @@ UObject* GetTopObject(const UTypedElementList* InElementList, const UClass* InRe
 	return nullptr;
 }
 
-UObject* GetBottomObject(const UTypedElementList* InElementList, const UClass* InRequiredClass)
+UObject* GetBottomObject(FTypedElementListConstRef InElementList, const UClass* InRequiredClass)
 {
 	TTypedElement<UTypedElementObjectInterface> TempElement;
 	for (int32 ElementIndex = InElementList->Num() - 1; ElementIndex >= 0; --ElementIndex)
@@ -98,17 +98,17 @@ UObject* GetBottomObject(const UTypedElementList* InElementList, const UClass* I
 	return nullptr;
 }
 
-bool HasObjectsOfExactClass(const UTypedElementList* InElementList, const UClass* InClass)
+bool HasObjectsOfExactClass(FTypedElementListConstRef InElementList, const UClass* InClass)
 {
 	return CountObjectsOfExactClass(InElementList, InClass) > 0;
 }
 
-int32 CountObjectsOfExactClass(const UTypedElementList* InElementList, const UClass* InClass)
+int32 CountObjectsOfExactClass(FTypedElementListConstRef InElementList, const UClass* InClass)
 {
 	return InElementList->GetCounter().GetCounterValue(NAME_Class, InClass);
 }
 
-void ForEachObjectClass(const UTypedElementList* InElementList, TFunctionRef<bool(UClass*)> InCallback)
+void ForEachObjectClass(FTypedElementListConstRef InElementList, TFunctionRef<bool(UClass*)> InCallback)
 {
 	InElementList->GetCounter().ForEachCounterValue<UClass*>(NAME_Class, [&InCallback](UClass* InClass, FTypedElementCounter::FCounterValue InCount)
 	{
diff --git a/Engine/Source/Runtime/TypedElementRuntime/Private/Elements/Framework/TypedElementSelectionSet.cpp b/Engine/Source/Runtime/TypedElementRuntime/Private/Elements/Framework/TypedElementSelectionSet.cpp
--- a/Engine/Source/Runtime/TypedElementRuntime/Private/Elements/Framework/TypedElementSelectionSet.cpp
+++ b/Engine/Source/Runtime/TypedElementRuntime/Private/Elements/Framework/TypedElementSelectionSet.cpp
@@ -8,7 +8,7 @@
 #include "Serialization/ObjectReader.h"
 #include "Serialization/ObjectWriter.h"
 
-void FTypedElementSelectionCustomization::GetNormalizedElements(const TTypedElement<UTypedElementSelectionInterface>& InElementSelectionHandle, const UTypedElementList* InSelectionSet, const FTypedElementSelectionNormalizationOptions& InNormalizationOptions, UTypedElementList* OutNormalizedElements)
+void FTypedElementSelectionCustomization::GetNormalizedElements(const TTypedElement<UTypedElementSelectionInterface>& InElementSelectionHandle, FTypedElementListConstRef InSelectionSet, const FTypedElementSelectionNormalizationOptions& InNormalizationOptions, FTypedElementListRef OutNormalizedElements)
 {
 	// Don't include this element in the normalized selection if it has a parent element that is also selected
 	{
@@ -201,7 +201,7 @@ bool UTypedElementSelectionSet::SelectElements(const TArray<FTypedElementHandle>
 
 bool UTypedElementSelectionSet::SelectElements(TArrayView<const FTypedElementHandle> InElementHandles, const FTypedElementSelectionOptions InSelectionOptions)
 {
-	FTypedElementListLegacySyncScopedBatch LegacySyncBatch(ElementList, InSelectionOptions.AllowLegacyNotifications());
+	FTypedElementListLegacySyncScopedBatch LegacySyncBatch(*ElementList, InSelectionOptions.AllowLegacyNotifications());
 
 	bool bSelectionChanged = false;
 
@@ -251,7 +251,7 @@ bool UTypedElementSelectionSet::DeselectElements(const TArray<FTypedElementHandl
 
 bool UTypedElementSelectionSet::DeselectElements(TArrayView<const FTypedElementHandle> InElementHandles, const FTypedElementSelectionOptions InSelectionOptions)
 {
-	FTypedElementListLegacySyncScopedBatch LegacySyncBatch(ElementList, InSelectionOptions.AllowLegacyNotifications());
+	FTypedElementListLegacySyncScopedBatch LegacySyncBatch(*ElementList, InSelectionOptions.AllowLegacyNotifications());
 
 	bool bSelectionChanged = false;
 
@@ -265,7 +265,7 @@ bool UTypedElementSelectionSet::DeselectElements(TArrayView<const FTypedElementH
 
 bool UTypedElementSelectionSet::ClearSelection(const FTypedElementSelectionOptions InSelectionOptions)
 {
-	FTypedElementListLegacySyncScopedBatch LegacySyncBatch(ElementList, InSelectionOptions.AllowLegacyNotifications());
+	FTypedElementListLegacySyncScopedBatch LegacySyncBatch(*ElementList, InSelectionOptions.AllowLegacyNotifications());
 
 	bool bSelectionChanged = false;
 
@@ -300,7 +300,7 @@ bool UTypedElementSelectionSet::SetSelection(const TArray<FTypedElementHandle>&
 
 bool UTypedElementSelectionSet::SetSelection(TArrayView<const FTypedElementHandle> InElementHandles, const FTypedElementSelectionOptions InSelectionOptions)
 {
-	FTypedElementListLegacySyncScopedBatch LegacySyncBatch(ElementList, InSelectionOptions.AllowLegacyNotifications());
+	FTypedElementListLegacySyncScopedBatch LegacySyncBatch(*ElementList, InSelectionOptions.AllowLegacyNotifications());
 
 	bool bSelectionChanged = false;
 
@@ -322,31 +322,31 @@ FTypedElementHandle UTypedElementSelectionSet::GetSelectionElement(const FTypedE
 	return SelectionSetElement ? SelectionSetElement.GetSelectionElement(InSelectionMethod) : FTypedElementHandle();
 }
 
-UTypedElementList* UTypedElementSelectionSet::GetNormalizedSelection(const FTypedElementSelectionNormalizationOptions InNormalizationOptions) const
+FTypedElementListRef UTypedElementSelectionSet::GetNormalizedSelection(const FTypedElementSelectionNormalizationOptions InNormalizationOptions) const
 {
-	return GetNormalizedElementList(ElementList, InNormalizationOptions);
+	return GetNormalizedElementList(ElementList.ToSharedRef(), InNormalizationOptions);
 }
 
-void UTypedElementSelectionSet::GetNormalizedSelection(const FTypedElementSelectionNormalizationOptions& InNormalizationOptions, UTypedElementList* OutNormalizedElements) const
+void UTypedElementSelectionSet::GetNormalizedSelection(const FTypedElementSelectionNormalizationOptions& InNormalizationOptions, FTypedElementListRef OutNormalizedElements) const
 {
-	GetNormalizedElementList(ElementList, InNormalizationOptions, OutNormalizedElements);
+	GetNormalizedElementList(ElementList.ToSharedRef(), InNormalizationOptions, OutNormalizedElements);
 }
 
-UTypedElementList* UTypedElementSelectionSet::GetNormalizedElementList(const UTypedElementList* InElementList, const FTypedElementSelectionNormalizationOptions InNormalizationOptions) const
+FTypedElementListRef UTypedElementSelectionSet::GetNormalizedElementList(FTypedElementListConstRef InElementList, const FTypedElementSelectionNormalizationOptions InNormalizationOptions) const
 {
 	UTypedElementRegistry* Registry = UTypedElementRegistry::GetInstance();
-	TGCObjectScopeGuard<UTypedElementList> NormalizedElementList(Registry->CreateElementList());
-	GetNormalizedElementList(InElementList, InNormalizationOptions, NormalizedElementList.Get());
-	return NormalizedElementList.Get();
+	FTypedElementListRef NormalizedElementList(Registry->CreateElementList());
+	GetNormalizedElementList(InElementList, InNormalizationOptions, NormalizedElementList);
+	return NormalizedElementList;
 }
 
-void UTypedElementSelectionSet::GetNormalizedElementList(const UTypedElementList* InElementList, const FTypedElementSelectionNormalizationOptions& InNormalizationOptions, UTypedElementList* OutNormalizedElements) const
+void UTypedElementSelectionSet::GetNormalizedElementList(FTypedElementListConstRef InElementList, const FTypedElementSelectionNormalizationOptions& InNormalizationOptions, FTypedElementListRef OutNormalizedElements) const
 {
 	OutNormalizedElements->Reset();
 	InElementList->ForEachElement<UTypedElementSelectionInterface>([this, InElementList, &InNormalizationOptions, OutNormalizedElements](const TTypedElement<UTypedElementSelectionInterface>& InElementSelectionHandle)
 	{
-		// Note: FTypedElementSelectionSetElement needs a non-const UTypedElementList, but we won't call anything that will modify it...
-		FTypedElementSelectionSetElement SelectionSetElement(InElementSelectionHandle, const_cast<UTypedElementList*>(InElementList), GetInterfaceCustomizationByTypeId(InElementSelectionHandle.GetId().GetTypeId()));
+		// Note: FTypedElementSelectionSetElement needs a non-const FTypedElementList, but we won't call anything that will modify it...
+		FTypedElementSelectionSetElement SelectionSetElement(InElementSelectionHandle, ConstCastSharedRef<FTypedElementList>(InElementList), GetInterfaceCustomizationByTypeId(InElementSelectionHandle.GetId().GetTypeId()));
 		check(SelectionSetElement.IsSet());
 		SelectionSetElement.GetNormalizedElements(InNormalizationOptions, OutNormalizedElements);
 		return true;
@@ -478,9 +478,9 @@ void UTypedElementSelectionSet::OnElementUpdated(TArrayView<const FTypedElementH
 	}
 }
 
-void UTypedElementSelectionSet::OnElementListPreChange(const UTypedElementList* InElementList)
+void UTypedElementSelectionSet::OnElementListPreChange(const FTypedElementList& InElementList)
 {
-	check(InElementList == ElementList);
+	check(&InElementList == ElementList.Get());
 	OnPreChangeDelegate.Broadcast(this);
 
 	if (!bIsRestoringState)
@@ -490,8 +490,8 @@ void UTypedElementSelectionSet::OnElementListPreChange(const UTypedElementList*
 	}
 }
 
-void UTypedElementSelectionSet::OnElementListChanged(const UTypedElementList* InElementList)
+void UTypedElementSelectionSet::OnElementListChanged(const FTypedElementList& InElementList)
 {
-	check(InElementList == ElementList);
+	check(&InElementList == ElementList.Get());
 	OnChangedDelegate.Broadcast(this);
 }
diff --git a/Engine/Source/Runtime/TypedElementRuntime/Private/Elements/Framework/TypedElementSelectionSetLibrary.h b/Engine/Source/Runtime/TypedElementRuntime/Private/Elements/Framework/TypedElementSelectionSetLibrary.h
new file mode 100644
--- /dev/null
+++ b/Engine/Source/Runtime/TypedElementRuntime/Private/Elements/Framework/TypedElementSelectionSetLibrary.h
@@ -0,0 +1,36 @@
+// Copyright Epic Games, Inc. All Rights Reserved.
+
+#pragma once
+
+#include "CoreMinimal.h"
+#include "Elements/Framework/TypedElementListProxy.h"
+#include "Elements/Framework/TypedElementSelectionSet.h"
+#include "TypedElementSelectionSetLibrary.generated.h"
+
+UCLASS()
+class UTypedElementSelectionSetLibrary : public UObject
+{
+	GENERATED_BODY()
+	
+public:
+	/**
+	 * Get a normalized version of this selection set that can be used to perform operations like gizmo manipulation, deletion, copying, etc.
+	 * This will do things like expand out groups, and resolve any parent<->child elements so that duplication operations aren't performed on both the parent and the child.
+	 */
+	UFUNCTION(BlueprintPure, Category="TypedElementFramework|Selection", meta=(ScriptMethod))
+	static FTypedElementListProxy GetNormalizedSelection(UTypedElementSelectionSet* SelectionSet, const FTypedElementSelectionNormalizationOptions NormalizationOptions)
+	{
+		return SelectionSet->GetNormalizedSelection(NormalizationOptions);
+	}
+
+	/**
+	 * Get a normalized version of the given element list that can be used to perform operations like gizmo manipulation, deletion, copying, etc.
+	 * This will do things like expand out groups, and resolve any parent<->child elements so that duplication operations aren't performed on both the parent and the child.
+	 */
+	UFUNCTION(BlueprintPure, Category="TypedElementFramework|Selection", meta=(ScriptMethod))
+	static FTypedElementListProxy GetNormalizedElementList(UTypedElementSelectionSet* SelectionSet, const FTypedElementListProxy ElementList, const FTypedElementSelectionNormalizationOptions NormalizationOptions)
+	{
+		FTypedElementListConstPtr ElementListPtr = ElementList.GetElementList();
+		return ElementListPtr ? SelectionSet->GetNormalizedElementList(ElementListPtr.ToSharedRef(), NormalizationOptions) : FTypedElementListProxy();
+	}
+};
diff --git a/Engine/Source/Runtime/TypedElementRuntime/Private/Elements/Interfaces/TypedElementSelectionInterface.cpp b/Engine/Source/Runtime/TypedElementRuntime/Private/Elements/Interfaces/TypedElementSelectionInterface.cpp
--- a/Engine/Source/Runtime/TypedElementRuntime/Private/Elements/Interfaces/TypedElementSelectionInterface.cpp
+++ b/Engine/Source/Runtime/TypedElementRuntime/Private/Elements/Interfaces/TypedElementSelectionInterface.cpp
@@ -3,17 +3,20 @@
 #include "Elements/Interfaces/TypedElementSelectionInterface.h"
 #include "Elements/Framework/TypedElementList.h"
 
-bool UTypedElementSelectionInterface::IsElementSelected(const FTypedElementHandle& InElementHandle, const UTypedElementList* InSelectionSet, const FTypedElementIsSelectedOptions& InSelectionOptions)
+bool UTypedElementSelectionInterface::IsElementSelected(const FTypedElementHandle& InElementHandle, const FTypedElementListProxy InSelectionSet, const FTypedElementIsSelectedOptions& InSelectionOptions)
 {
-	return InSelectionSet->Contains(InElementHandle);
+	FTypedElementListConstPtr SelectionSetPtr = InSelectionSet.GetElementList();
+	return SelectionSetPtr && SelectionSetPtr->Contains(InElementHandle);
 }
 
-bool UTypedElementSelectionInterface::SelectElement(const FTypedElementHandle& InElementHandle, UTypedElementList* InSelectionSet, const FTypedElementSelectionOptions& InSelectionOptions)
+bool UTypedElementSelectionInterface::SelectElement(const FTypedElementHandle& InElementHandle, FTypedElementListProxy InSelectionSet, const FTypedElementSelectionOptions& InSelectionOptions)
 {
-	return InSelectionSet->Add(InElementHandle);
+	FTypedElementListPtr SelectionSetPtr = InSelectionSet.GetElementList();
+	return SelectionSetPtr && SelectionSetPtr->Add(InElementHandle);
 }
 
-bool UTypedElementSelectionInterface::DeselectElement(const FTypedElementHandle& InElementHandle, UTypedElementList* InSelectionSet, const FTypedElementSelectionOptions& InSelectionOptions)
+bool UTypedElementSelectionInterface::DeselectElement(const FTypedElementHandle& InElementHandle, FTypedElementListProxy InSelectionSet, const FTypedElementSelectionOptions& InSelectionOptions)
 {
-	return InSelectionSet->Remove(InElementHandle);
+	FTypedElementListPtr SelectionSetPtr = InSelectionSet.GetElementList();
+	return SelectionSetPtr && SelectionSetPtr->Remove(InElementHandle);
 }
diff --git a/Engine/Source/Runtime/TypedElementRuntime/Public/Elements/Framework/TypedElementListObjectUtil.h b/Engine/Source/Runtime/TypedElementRuntime/Public/Elements/Framework/TypedElementListObjectUtil.h
--- a/Engine/Source/Runtime/TypedElementRuntime/Public/Elements/Framework/TypedElementListObjectUtil.h
+++ b/Engine/Source/Runtime/TypedElementRuntime/Public/Elements/Framework/TypedElementListObjectUtil.h
@@ -10,13 +10,13 @@ namespace TypedElementListObjectUtil
 /**
  * Test whether there are any objects in the given list of elements.
  */
-TYPEDELEMENTRUNTIME_API bool HasObjects(const UTypedElementList* InElementList, const UClass* InRequiredClass = nullptr);
+TYPEDELEMENTRUNTIME_API bool HasObjects(FTypedElementListConstRef InElementList, const UClass* InRequiredClass = nullptr);
 
 /**
  * Test whether there are any objects in the given list of elements.
  */
 template <typename RequiredClassType>
-bool HasObjects(const UTypedElementList* InElementList)
+bool HasObjects(FTypedElementListConstRef InElementList)
 {
 	return HasObjects(InElementList, RequiredClassType::StaticClass());
 }
@@ -24,13 +24,13 @@ bool HasObjects(const UTypedElementList* InElementList)
 /**
  * Count the number of objects in the given list of elements.
  */
-TYPEDELEMENTRUNTIME_API int32 CountObjects(const UTypedElementList* InElementList, const UClass* InRequiredClass = nullptr);
+TYPEDELEMENTRUNTIME_API int32 CountObjects(FTypedElementListConstRef InElementList, const UClass* InRequiredClass = nullptr);
 
 /**
  * Count the number of objects in the given list of elements.
  */
 template <typename RequiredClassType>
-int32 CountObjects(const UTypedElementList* InElementList)
+int32 CountObjects(FTypedElementListConstRef InElementList)
 {
 	return CountObjects(InElementList, RequiredClassType::StaticClass());
 }
@@ -39,14 +39,14 @@ int32 CountObjects(const UTypedElementList* InElementList)
  * Enumerate the objects from the given list of elements.
  * @note Return true from the callback to continue enumeration.
  */
-TYPEDELEMENTRUNTIME_API void ForEachObject(const UTypedElementList* InElementList, TFunctionRef<bool(UObject*)> InCallback, const UClass* InRequiredClass = nullptr);
+TYPEDELEMENTRUNTIME_API void ForEachObject(FTypedElementListConstRef InElementList, TFunctionRef<bool(UObject*)> InCallback, const UClass* InRequiredClass = nullptr);
 
 /**
  * Enumerate the objects from the given list of elements.
  * @note Return true from the callback to continue enumeration.
  */
 template <typename RequiredClassType>
-void ForEachObject(const UTypedElementList* InElementList, TFunctionRef<bool(RequiredClassType*)> InCallback)
+void ForEachObject(FTypedElementListConstRef InElementList, TFunctionRef<bool(RequiredClassType*)> InCallback)
 {
 	ForEachObject(InElementList, [&InCallback](UObject* InObject)
 	{
@@ -57,13 +57,13 @@ void ForEachObject(const UTypedElementList* InElementList, TFunctionRef<bool(Req
 /**
  * Get the array of objects from the given list of elements.
  */
-TYPEDELEMENTRUNTIME_API TArray<UObject*> GetObjects(const UTypedElementList* InElementList, const UClass* InRequiredClass = nullptr);
+TYPEDELEMENTRUNTIME_API TArray<UObject*> GetObjects(FTypedElementListConstRef InElementList, const UClass* InRequiredClass = nullptr);
 
 /**
  * Get the array of objects from the given list of elements.
  */
 template <typename RequiredClassType>
-TArray<RequiredClassType*> GetObjects(const UTypedElementList* InElementList)
+TArray<RequiredClassType*> GetObjects(FTypedElementListConstRef InElementList)
 {
 	TArray<RequiredClassType*> SelectedObjects;
 	SelectedObjects.Reserve(InElementList->Num());
@@ -80,13 +80,13 @@ TArray<RequiredClassType*> GetObjects(const UTypedElementList* InElementList)
 /**
  * Get the first object of the given type from the given list of elements.
  */
-TYPEDELEMENTRUNTIME_API UObject* GetTopObject(const UTypedElementList* InElementList, const UClass* InRequiredClass = nullptr);
+TYPEDELEMENTRUNTIME_API UObject* GetTopObject(FTypedElementListConstRef InElementList, const UClass* InRequiredClass = nullptr);
 
 /**
  * Get the first object of the given type from the given list of elements.
  */
 template <typename RequiredClassType>
-RequiredClassType* GetTopObject(const UTypedElementList* InElementList)
+RequiredClassType* GetTopObject(FTypedElementListConstRef InElementList)
 {
 	return Cast<RequiredClassType>(GetTopObject(InElementList, RequiredClassType::StaticClass()));
 }
@@ -94,13 +94,13 @@ RequiredClassType* GetTopObject(const UTypedElementList* InElementList)
 /**
  * Get the last object of the given type from the given list of elements.
  */
-TYPEDELEMENTRUNTIME_API UObject* GetBottomObject(const UTypedElementList* InElementList, const UClass* InRequiredClass = nullptr);
+TYPEDELEMENTRUNTIME_API UObject* GetBottomObject(FTypedElementListConstRef InElementList, const UClass* InRequiredClass = nullptr);
 
 /**
  * Get the last object of the given type from the given list of elements.
  */
 template <typename RequiredClassType>
-RequiredClassType* GetBottomObject(const UTypedElementList* InElementList)
+RequiredClassType* GetBottomObject(FTypedElementListConstRef InElementList)
 {
 	return Cast<RequiredClassType>(GetBottomObject(InElementList, RequiredClassType::StaticClass()));
 }
@@ -108,13 +108,13 @@ RequiredClassType* GetBottomObject(const UTypedElementList* InElementList)
 /**
  * Test if there are any objects of the exact class in the given list of elements (a quick test using the class counter, skipping derived types).
  */
-TYPEDELEMENTRUNTIME_API bool HasObjectsOfExactClass(const UTypedElementList* InElementList, const UClass* InClass);
+TYPEDELEMENTRUNTIME_API bool HasObjectsOfExactClass(FTypedElementListConstRef InElementList, const UClass* InClass);
 
 /**
  * Test if there are any objects of the exact class in the given list of elements (a quick test using the class counter, skipping derived types).
  */
 template <typename ClassType>
-bool HasObjectsOfExactClass(const UTypedElementList* InElementList)
+bool HasObjectsOfExactClass(FTypedElementListConstRef InElementList)
 {
 	return HasObjectsOfExactClass(InElementList, ClassType::StaticClass());
 }
@@ -122,13 +122,13 @@ bool HasObjectsOfExactClass(const UTypedElementList* InElementList)
 /**
  * Count the number of objects of the exact class in the given list of elements (a quick test using the class counter, skipping derived types).
  */
-TYPEDELEMENTRUNTIME_API int32 CountObjectsOfExactClass(const UTypedElementList* InElementList, const UClass* InClass);
+TYPEDELEMENTRUNTIME_API int32 CountObjectsOfExactClass(FTypedElementListConstRef InElementList, const UClass* InClass);
 
 /**
  * Count the number of objects of the exact class in the given list of elements (a quick test using the class counter, skipping derived types).
  */
 template <typename ClassType>
-bool CountObjectsOfExactClass(const UTypedElementList* InElementList)
+bool CountObjectsOfExactClass(FTypedElementListConstRef InElementList)
 {
 	return CountObjectsOfExactClass(InElementList, ClassType::StaticClass());
 }
@@ -137,6 +137,6 @@ bool CountObjectsOfExactClass(const UTypedElementList* InElementList)
  * Enumerate the classes of the objects in the given list of elements.
  * @note Return true from the callback to continue enumeration.
  */
-TYPEDELEMENTRUNTIME_API void ForEachObjectClass(const UTypedElementList* InElementList, TFunctionRef<bool(UClass*)> InCallback);
+TYPEDELEMENTRUNTIME_API void ForEachObjectClass(FTypedElementListConstRef InElementList, TFunctionRef<bool(UClass*)> InCallback);
 
 } // namespace TypedElementListObjectUtil
diff --git a/Engine/Source/Runtime/TypedElementRuntime/Public/Elements/Framework/TypedElementSelectionSet.h b/Engine/Source/Runtime/TypedElementRuntime/Public/Elements/Framework/TypedElementSelectionSet.h
--- a/Engine/Source/Runtime/TypedElementRuntime/Public/Elements/Framework/TypedElementSelectionSet.h
+++ b/Engine/Source/Runtime/TypedElementRuntime/Public/Elements/Framework/TypedElementSelectionSet.h
@@ -38,14 +38,14 @@ public:
 	virtual ~FTypedElementSelectionCustomization() = default;
 
 	//~ See UTypedElementSelectionInterface for API docs
-	virtual bool IsElementSelected(const TTypedElement<UTypedElementSelectionInterface>& InElementSelectionHandle, const UTypedElementList* InSelectionSet, const FTypedElementIsSelectedOptions& InSelectionOptions) { return InElementSelectionHandle.IsElementSelected(InSelectionSet, InSelectionOptions); }
+	virtual bool IsElementSelected(const TTypedElement<UTypedElementSelectionInterface>& InElementSelectionHandle, FTypedElementListConstRef InSelectionSet, const FTypedElementIsSelectedOptions& InSelectionOptions) { return InElementSelectionHandle.IsElementSelected(InSelectionSet, InSelectionOptions); }
 	virtual bool CanSelectElement(const TTypedElement<UTypedElementSelectionInterface>& InElementSelectionHandle, const FTypedElementSelectionOptions& InSelectionOptions) { return InElementSelectionHandle.CanSelectElement(InSelectionOptions); }
 	virtual bool CanDeselectElement(const TTypedElement<UTypedElementSelectionInterface>& InElementSelectionHandle, const FTypedElementSelectionOptions& InSelectionOptions) { return InElementSelectionHandle.CanDeselectElement(InSelectionOptions); }
-	virtual bool SelectElement(const TTypedElement<UTypedElementSelectionInterface>& InElementSelectionHandle, UTypedElementList* InSelectionSet, const FTypedElementSelectionOptions& InSelectionOptions) { return InElementSelectionHandle.SelectElement(InSelectionSet, InSelectionOptions); }
-	virtual bool DeselectElement(const TTypedElement<UTypedElementSelectionInterface>& InElementSelectionHandle, UTypedElementList* InSelectionSet, const FTypedElementSelectionOptions& InSelectionOptions) { return InElementSelectionHandle.DeselectElement(InSelectionSet, InSelectionOptions); }
-	virtual bool AllowSelectionModifiers(const TTypedElement<UTypedElementSelectionInterface>& InElementSelectionHandle, const UTypedElementList* InSelectionSet) { return InElementSelectionHandle.AllowSelectionModifiers(InSelectionSet); }
-	virtual FTypedElementHandle GetSelectionElement(const TTypedElement<UTypedElementSelectionInterface>& InElementSelectionHandle, const UTypedElementList* InCurrentSelection, const ETypedElementSelectionMethod InSelectionMethod) { return InElementSelectionHandle.GetSelectionElement(InCurrentSelection, InSelectionMethod); }
-	virtual void GetNormalizedElements(const TTypedElement<UTypedElementSelectionInterface>& InElementSelectionHandle, const UTypedElementList* InSelectionSet, const FTypedElementSelectionNormalizationOptions& InNormalizationOptions, UTypedElementList* OutNormalizedElements);
+	virtual bool SelectElement(const TTypedElement<UTypedElementSelectionInterface>& InElementSelectionHandle, FTypedElementListRef InSelectionSet, const FTypedElementSelectionOptions& InSelectionOptions) { return InElementSelectionHandle.SelectElement(InSelectionSet, InSelectionOptions); }
+	virtual bool DeselectElement(const TTypedElement<UTypedElementSelectionInterface>& InElementSelectionHandle, FTypedElementListRef InSelectionSet, const FTypedElementSelectionOptions& InSelectionOptions) { return InElementSelectionHandle.DeselectElement(InSelectionSet, InSelectionOptions); }
+	virtual bool AllowSelectionModifiers(const TTypedElement<UTypedElementSelectionInterface>& InElementSelectionHandle, FTypedElementListConstRef InSelectionSet) { return InElementSelectionHandle.AllowSelectionModifiers(InSelectionSet); }
+	virtual FTypedElementHandle GetSelectionElement(const TTypedElement<UTypedElementSelectionInterface>& InElementSelectionHandle, FTypedElementListConstRef InCurrentSelection, const ETypedElementSelectionMethod InSelectionMethod) { return InElementSelectionHandle.GetSelectionElement(InCurrentSelection, InSelectionMethod); }
+	virtual void GetNormalizedElements(const TTypedElement<UTypedElementSelectionInterface>& InElementSelectionHandle, FTypedElementListConstRef InSelectionSet, const FTypedElementSelectionNormalizationOptions& InNormalizationOptions, FTypedElementListRef OutNormalizedElements);
 };
 
 /**
@@ -56,7 +56,7 @@ struct TYPEDELEMENTRUNTIME_API FTypedElementSelectionSetElement
 public:
 	FTypedElementSelectionSetElement() = default;
 
-	FTypedElementSelectionSetElement(TTypedElement<UTypedElementSelectionInterface> InElementSelectionHandle, UTypedElementList* InElementList, FTypedElementSelectionCustomization* InSelectionCustomization)
+	FTypedElementSelectionSetElement(TTypedElement<UTypedElementSelectionInterface> InElementSelectionHandle, FTypedElementListPtr InElementList, FTypedElementSelectionCustomization* InSelectionCustomization)
 		: ElementSelectionHandle(MoveTemp(InElementSelectionHandle))
 		, ElementList(InElementList)
 		, SelectionCustomization(InSelectionCustomization)
@@ -82,18 +82,18 @@ public:
 	}
 
 	//~ See UTypedElementSelectionInterface for API docs
-	bool IsElementSelected(const FTypedElementIsSelectedOptions& InSelectionOptions) const { return SelectionCustomization->IsElementSelected(ElementSelectionHandle, ElementList, InSelectionOptions); }
+	bool IsElementSelected(const FTypedElementIsSelectedOptions& InSelectionOptions) const { return SelectionCustomization->IsElementSelected(ElementSelectionHandle, ElementList.ToSharedRef(), InSelectionOptions); }
 	bool CanSelectElement(const FTypedElementSelectionOptions& InSelectionOptions) const { return SelectionCustomization->CanSelectElement(ElementSelectionHandle, InSelectionOptions); }
 	bool CanDeselectElement(const FTypedElementSelectionOptions& InSelectionOptions) const { return SelectionCustomization->CanDeselectElement(ElementSelectionHandle, InSelectionOptions); }
-	bool SelectElement(const FTypedElementSelectionOptions& InSelectionOptions) const { return SelectionCustomization->SelectElement(ElementSelectionHandle, ElementList, InSelectionOptions); }
-	bool DeselectElement(const FTypedElementSelectionOptions& InSelectionOptions) const { return SelectionCustomization->DeselectElement(ElementSelectionHandle, ElementList, InSelectionOptions); }
-	bool AllowSelectionModifiers() const { return SelectionCustomization->AllowSelectionModifiers(ElementSelectionHandle, ElementList); }
-	FTypedElementHandle GetSelectionElement(const ETypedElementSelectionMethod InSelectionMethod) const { return SelectionCustomization->GetSelectionElement(ElementSelectionHandle, ElementList, InSelectionMethod); }
-	void GetNormalizedElements(const FTypedElementSelectionNormalizationOptions& InNormalizationOptions, UTypedElementList* OutNormalizedElements) const { return SelectionCustomization->GetNormalizedElements(ElementSelectionHandle, ElementList, InNormalizationOptions, OutNormalizedElements); }
+	bool SelectElement(const FTypedElementSelectionOptions& InSelectionOptions) const { return SelectionCustomization->SelectElement(ElementSelectionHandle, ElementList.ToSharedRef(), InSelectionOptions); }
+	bool DeselectElement(const FTypedElementSelectionOptions& InSelectionOptions) const { return SelectionCustomization->DeselectElement(ElementSelectionHandle, ElementList.ToSharedRef(), InSelectionOptions); }
+	bool AllowSelectionModifiers() const { return SelectionCustomization->AllowSelectionModifiers(ElementSelectionHandle, ElementList.ToSharedRef()); }
+	FTypedElementHandle GetSelectionElement(const ETypedElementSelectionMethod InSelectionMethod) const { return SelectionCustomization->GetSelectionElement(ElementSelectionHandle, ElementList.ToSharedRef(), InSelectionMethod); }
+	void GetNormalizedElements(const FTypedElementSelectionNormalizationOptions& InNormalizationOptions, FTypedElementListRef OutNormalizedElements) const { return SelectionCustomization->GetNormalizedElements(ElementSelectionHandle, ElementList.ToSharedRef(), InNormalizationOptions, OutNormalizedElements); }
 
 private:
 	TTypedElement<UTypedElementSelectionInterface> ElementSelectionHandle;
-	UTypedElementList* ElementList = nullptr;
+	FTypedElementListPtr ElementList;
 	FTypedElementSelectionCustomization* SelectionCustomization = nullptr;
 };
 
@@ -241,17 +241,15 @@ public:
 	 * Get a normalized version of this selection set that can be used to perform operations like gizmo manipulation, deletion, copying, etc.
 	 * This will do things like expand out groups, and resolve any parent<->child elements so that duplication operations aren't performed on both the parent and the child.
 	 */
-	UFUNCTION(BlueprintPure, Category="TypedElementFramework|Selection")
-	UTypedElementList* GetNormalizedSelection(const FTypedElementSelectionNormalizationOptions InNormalizationOptions) const;
-	void GetNormalizedSelection(const FTypedElementSelectionNormalizationOptions& InNormalizationOptions, UTypedElementList* OutNormalizedElements) const;
+	FTypedElementListRef GetNormalizedSelection(const FTypedElementSelectionNormalizationOptions InNormalizationOptions) const;
+	void GetNormalizedSelection(const FTypedElementSelectionNormalizationOptions& InNormalizationOptions, FTypedElementListRef OutNormalizedElements) const;
 
 	/**
 	 * Get a normalized version of the given element list that can be used to perform operations like gizmo manipulation, deletion, copying, etc.
 	 * This will do things like expand out groups, and resolve any parent<->child elements so that duplication operations aren't performed on both the parent and the child.
 	 */
-	UFUNCTION(BlueprintPure, Category="TypedElementFramework|Selection")
-	UTypedElementList* GetNormalizedElementList(const UTypedElementList* InElementList, const FTypedElementSelectionNormalizationOptions InNormalizationOptions) const;
-	void GetNormalizedElementList(const UTypedElementList* InElementList, const FTypedElementSelectionNormalizationOptions& InNormalizationOptions, UTypedElementList* OutNormalizedElements) const;
+	FTypedElementListRef GetNormalizedElementList(FTypedElementListConstRef InElementList, const FTypedElementSelectionNormalizationOptions InNormalizationOptions) const;
+	void GetNormalizedElementList(FTypedElementListConstRef InElementList, const FTypedElementSelectionNormalizationOptions& InNormalizationOptions, FTypedElementListRef OutNormalizedElements) const;
 
 	/**
 	 * Get the number of selected elements.
@@ -341,7 +339,7 @@ public:
 	UFUNCTION(BlueprintCallable, BlueprintPure=false, Category="TypedElementFramework|Selection")
 	bool HasSelectedObjects(const UClass* InRequiredClass = nullptr) const
 	{
-		return TypedElementListObjectUtil::HasObjects(ElementList, InRequiredClass);
+		return TypedElementListObjectUtil::HasObjects(ElementList.ToSharedRef(), InRequiredClass);
 	}
 
 	/**
@@ -350,7 +348,7 @@ public:
 	template <typename RequiredClassType>
 	bool HasSelectedObjects() const
 	{
-		return TypedElementListObjectUtil::HasObjects<RequiredClassType>(ElementList);
+		return TypedElementListObjectUtil::HasObjects<RequiredClassType>(ElementList.ToSharedRef());
 	}
 
 	/**
@@ -359,7 +357,7 @@ public:
 	UFUNCTION(BlueprintCallable, BlueprintPure=false, Category="TypedElementFramework|Selection")
 	int32 CountSelectedObjects(const UClass* InRequiredClass = nullptr) const
 	{
-		return TypedElementListObjectUtil::CountObjects(ElementList, InRequiredClass);
+		return TypedElementListObjectUtil::CountObjects(ElementList.ToSharedRef(), InRequiredClass);
 	}
 
 	/**
@@ -368,7 +366,7 @@ public:
 	template <typename RequiredClassType>
 	int32 CountSelectedObjects() const
 	{
-		return TypedElementListObjectUtil::CountObjects<RequiredClassType>(ElementList);
+		return TypedElementListObjectUtil::CountObjects<RequiredClassType>(ElementList.ToSharedRef());
 	}
 
 	/**
@@ -377,7 +375,7 @@ public:
 	UFUNCTION(BlueprintCallable, BlueprintPure=false, Category="TypedElementFramework|Selection")
 	TArray<UObject*> GetSelectedObjects(const UClass* InRequiredClass = nullptr) const
 	{
-		return TypedElementListObjectUtil::GetObjects(ElementList, InRequiredClass);
+		return TypedElementListObjectUtil::GetObjects(ElementList.ToSharedRef(), InRequiredClass);
 	}
 
 	/**
@@ -386,7 +384,7 @@ public:
 	template <typename RequiredClassType>
 	TArray<RequiredClassType*> GetSelectedObjects() const
 	{
-		return TypedElementListObjectUtil::GetObjects<RequiredClassType>(ElementList);
+		return TypedElementListObjectUtil::GetObjects<RequiredClassType>(ElementList.ToSharedRef());
 	}
 
 	/**
@@ -395,7 +393,7 @@ public:
 	 */
 	void ForEachSelectedObject(TFunctionRef<bool(UObject*)> InCallback, const UClass* InRequiredClass = nullptr) const
 	{
-		TypedElementListObjectUtil::ForEachObject(ElementList, InCallback, InRequiredClass);
+		TypedElementListObjectUtil::ForEachObject(ElementList.ToSharedRef(), InCallback, InRequiredClass);
 	}
 
 	/**
@@ -405,7 +403,7 @@ public:
 	template <typename RequiredClassType>
 	void ForEachSelectedObject(TFunctionRef<bool(RequiredClassType*)> InCallback) const
 	{
-		TypedElementListObjectUtil::ForEachObject<RequiredClassType>(ElementList, InCallback);
+		TypedElementListObjectUtil::ForEachObject<RequiredClassType>(ElementList.ToSharedRef(), InCallback);
 	}
 
 	/**
@@ -414,7 +412,7 @@ public:
 	UFUNCTION(BlueprintPure, Category="TypedElementFramework|Selection")
 	UObject* GetTopSelectedObject(const UClass* InRequiredClass = nullptr) const
 	{
-		return TypedElementListObjectUtil::GetTopObject(ElementList, InRequiredClass);
+		return TypedElementListObjectUtil::GetTopObject(ElementList.ToSharedRef(), InRequiredClass);
 	}
 
 	/**
@@ -423,7 +421,7 @@ public:
 	template <typename RequiredClassType>
 	RequiredClassType* GetTopSelectedObject() const
 	{
-		return TypedElementListObjectUtil::GetTopObject<RequiredClassType>(ElementList);
+		return TypedElementListObjectUtil::GetTopObject<RequiredClassType>(ElementList.ToSharedRef());
 	}
 
 	/**
@@ -432,7 +430,7 @@ public:
 	UFUNCTION(BlueprintPure, Category="TypedElementFramework|Selection")
 	UObject* GetBottomSelectedObject(const UClass* InRequiredClass = nullptr) const
 	{
-		return TypedElementListObjectUtil::GetBottomObject(ElementList, InRequiredClass);
+		return TypedElementListObjectUtil::GetBottomObject(ElementList.ToSharedRef(), InRequiredClass);
 	}
 
 	/**
@@ -441,7 +439,7 @@ public:
 	template <typename RequiredClassType>
 	RequiredClassType* GetBottomSelectedObject() const
 	{
-		return TypedElementListObjectUtil::GetBottomObject<RequiredClassType>(ElementList);
+		return TypedElementListObjectUtil::GetBottomObject<RequiredClassType>(ElementList.ToSharedRef());
 	}
 
 	/**
@@ -501,9 +499,9 @@ public:
 	/**
 	 * Get the underlying element list holding the selection state.
 	 */
-	const UTypedElementList* GetElementList() const
+	FTypedElementListConstRef GetElementList() const
 	{
-		return ElementList;
+		return ElementList.ToSharedRef();
 	}
 
 	/**
@@ -541,16 +539,15 @@ private:
 	/**
 	 * Proxy the internal OnPreChange event from the underlying element list.
 	 */
-	void OnElementListPreChange(const UTypedElementList* InElementList);
+	void OnElementListPreChange(const FTypedElementList& InElementList);
 
 	/**
 	 * Proxy the internal OnChanged event from the underlying element list.
 	 */
-	void OnElementListChanged(const UTypedElementList* InElementList);
+	void OnElementListChanged(const FTypedElementList& InElementList);
 
 	/** Underlying element list holding the selection state. */
-	UPROPERTY()
-	TObjectPtr<UTypedElementList> ElementList = nullptr;
+	FTypedElementListPtr ElementList;
 
 	/** Delegate that is invoked whenever the underlying element list is potentially about to change. */
 	FOnPreChange OnPreChangeDelegate;
diff --git a/Engine/Source/Runtime/TypedElementRuntime/Public/Elements/Interfaces/TypedElementSelectionInterface.h b/Engine/Source/Runtime/TypedElementRuntime/Public/Elements/Interfaces/TypedElementSelectionInterface.h
--- a/Engine/Source/Runtime/TypedElementRuntime/Public/Elements/Interfaces/TypedElementSelectionInterface.h
+++ b/Engine/Source/Runtime/TypedElementRuntime/Public/Elements/Interfaces/TypedElementSelectionInterface.h
@@ -3,10 +3,9 @@
 #pragma once
 
 #include "Elements/Framework/TypedElementHandle.h"
+#include "Elements/Framework/TypedElementListProxy.h"
 #include "TypedElementSelectionInterface.generated.h"
 
-class UTypedElementList;
-
 UENUM()
 enum class ETypedElementSelectionMethod : uint8
 {
@@ -138,7 +137,7 @@ public:
 	 * Test to see whether the given element is currently considered selected.
 	 */
 	UFUNCTION(BlueprintPure, Category="TypedElementInterfaces|Selection")
-	virtual bool IsElementSelected(const FTypedElementHandle& InElementHandle, const UTypedElementList* InSelectionSet, const FTypedElementIsSelectedOptions& InSelectionOptions);
+	virtual bool IsElementSelected(const FTypedElementHandle& InElementHandle, const FTypedElementListProxy InSelectionSet, const FTypedElementIsSelectedOptions& InSelectionOptions);
 
 	/**
 	 * Test to see whether the given element can be selected.
@@ -157,26 +156,26 @@ public:
 	 * @return True if the selection was changed, false otherwise.
 	 */
 	UFUNCTION(BlueprintCallable, Category="TypedElementInterfaces|Selection")
-	virtual bool SelectElement(const FTypedElementHandle& InElementHandle, UTypedElementList* InSelectionSet, const FTypedElementSelectionOptions& InSelectionOptions);
+	virtual bool SelectElement(const FTypedElementHandle& InElementHandle, FTypedElementListProxy InSelectionSet, const FTypedElementSelectionOptions& InSelectionOptions);
 
 	/**
 	 * Attempt to deselect the given element.
 	 * @return True if the selection was changed, false otherwise.
 	 */
 	UFUNCTION(BlueprintCallable, Category="TypedElementInterfaces|Selection")
-	virtual bool DeselectElement(const FTypedElementHandle& InElementHandle, UTypedElementList* InSelectionSet, const FTypedElementSelectionOptions& InSelectionOptions);
+	virtual bool DeselectElement(const FTypedElementHandle& InElementHandle, FTypedElementListProxy InSelectionSet, const FTypedElementSelectionOptions& InSelectionOptions);
 
 	/**
 	 * Test to see whether selection modifiers (Ctrl or Shift) are allowed while selecting this element.
 	 */
 	UFUNCTION(BlueprintPure, Category="TypedElementInterfaces|Selection")
-	virtual bool AllowSelectionModifiers(const FTypedElementHandle& InElementHandle, const UTypedElementList* InSelectionSet) { return true; }
+	virtual bool AllowSelectionModifiers(const FTypedElementHandle& InElementHandle, const FTypedElementListProxy InSelectionSet) { return true; }
 
 	/**
 	 * Given an element, return the element that should actually perform a selection operation.
 	 */
 	UFUNCTION(BlueprintPure, Category="TypedElementInterfaces|Selection")
-	virtual FTypedElementHandle GetSelectionElement(const FTypedElementHandle& InElementHandle, const UTypedElementList* InCurrentSelection, const ETypedElementSelectionMethod InSelectionMethod) { return InElementHandle; }
+	virtual FTypedElementHandle GetSelectionElement(const FTypedElementHandle& InElementHandle, const FTypedElementListProxy InCurrentSelection, const ETypedElementSelectionMethod InSelectionMethod) { return InElementHandle; }
 
 	/**
 	 * Test to see whether the given element prevents the selection set state from being transacted for undo/redo (eg, if the element belongs to a PIE instance).
@@ -220,13 +219,13 @@ protected:
 template <>
 struct TTypedElement<UTypedElementSelectionInterface> : public TTypedElementBase<UTypedElementSelectionInterface>
 {
-	bool IsElementSelected(const UTypedElementList* InSelectionSet, const FTypedElementIsSelectedOptions& InSelectionOptions) const { return InterfacePtr->IsElementSelected(*this, InSelectionSet, InSelectionOptions); }
+	bool IsElementSelected(FTypedElementListConstRef InSelectionSet, const FTypedElementIsSelectedOptions& InSelectionOptions) const { return InterfacePtr->IsElementSelected(*this, InSelectionSet, InSelectionOptions); }
 	bool CanSelectElement(const FTypedElementSelectionOptions& InSelectionOptions) const { return InterfacePtr->CanSelectElement(*this, InSelectionOptions); }
 	bool CanDeselectElement(const FTypedElementSelectionOptions& InSelectionOptions) const { return InterfacePtr->CanDeselectElement(*this, InSelectionOptions); }
-	bool SelectElement(UTypedElementList* InSelectionSet, const FTypedElementSelectionOptions& InSelectionOptions) const { return InterfacePtr->SelectElement(*this, InSelectionSet, InSelectionOptions); }
-	bool DeselectElement(UTypedElementList* InSelectionSet, const FTypedElementSelectionOptions& InSelectionOptions) const { return InterfacePtr->DeselectElement(*this, InSelectionSet, InSelectionOptions); }
-	bool AllowSelectionModifiers(const UTypedElementList* InSelectionSet) const { return InterfacePtr->AllowSelectionModifiers(*this, InSelectionSet); }
-	FTypedElementHandle GetSelectionElement(const UTypedElementList* InCurrentSelection, const ETypedElementSelectionMethod InSelectionMethod) const { return InterfacePtr->GetSelectionElement(*this, InCurrentSelection, InSelectionMethod); }
+	bool SelectElement(FTypedElementListRef InSelectionSet, const FTypedElementSelectionOptions& InSelectionOptions) const { return InterfacePtr->SelectElement(*this, InSelectionSet, InSelectionOptions); }
+	bool DeselectElement(FTypedElementListRef InSelectionSet, const FTypedElementSelectionOptions& InSelectionOptions) const { return InterfacePtr->DeselectElement(*this, InSelectionSet, InSelectionOptions); }
+	bool AllowSelectionModifiers(FTypedElementListConstRef InSelectionSet) const { return InterfacePtr->AllowSelectionModifiers(*this, InSelectionSet); }
+	FTypedElementHandle GetSelectionElement(FTypedElementListConstRef InCurrentSelection, const ETypedElementSelectionMethod InSelectionMethod) const { return InterfacePtr->GetSelectionElement(*this, InCurrentSelection, InSelectionMethod); }
 	bool ShouldPreventTransactions() const { return InterfacePtr->ShouldPreventTransactions(*this); }
 	TUniquePtr<ITypedElementTransactedElement> CreateTransactedElement() const { return InterfacePtr->CreateTransactedElement(*this); }
 };
