Fixed lingering element references when destroying a selection set

diff --git a/Engine/Source/Runtime/TypedElementRuntime/Private/Elements/Framework/TypedElementSelectionSet.cpp b/Engine/Source/Runtime/TypedElementRuntime/Private/Elements/Framework/TypedElementSelectionSet.cpp
--- a/Engine/Source/Runtime/TypedElementRuntime/Private/Elements/Framework/TypedElementSelectionSet.cpp
+++ b/Engine/Source/Runtime/TypedElementRuntime/Private/Elements/Framework/TypedElementSelectionSet.cpp
@@ -86,6 +86,16 @@ bool UTypedElementSelectionSet::Modify(bool bAlwaysMarkDirty)
 }
 #endif	// WITH_EDITOR
 
+void UTypedElementSelectionSet::BeginDestroy()
+{
+	if (ElementList)
+	{
+		ElementList->Empty();
+	}
+
+	Super::BeginDestroy();
+}
+
 void UTypedElementSelectionSet::Serialize(FArchive& Ar)
 {
 	checkf(!Ar.IsPersistent(), TEXT("UTypedElementSelectionSet can only be serialized by transient archives!"));
diff --git a/Engine/Source/Runtime/TypedElementRuntime/Public/Elements/Framework/TypedElementSelectionSet.h b/Engine/Source/Runtime/TypedElementRuntime/Public/Elements/Framework/TypedElementSelectionSet.h
--- a/Engine/Source/Runtime/TypedElementRuntime/Public/Elements/Framework/TypedElementSelectionSet.h
+++ b/Engine/Source/Runtime/TypedElementRuntime/Public/Elements/Framework/TypedElementSelectionSet.h
@@ -159,6 +159,7 @@ public:
 	virtual void PostEditUndo() override;
 	virtual bool Modify(bool bAlwaysMarkDirty = true) override;
 #endif	// WITH_EDITOR
+	virtual void BeginDestroy() override;
 	virtual void Serialize(FArchive& Ar) override;
 
 	/**
