Fixed an issue with the UTypedElementSelectionSet not iterring on the objects that implement a certain class.

diff --git a/Engine/Source/Runtime/TypedElementRuntime/Private/Elements/Framework/TypedElementListObjectUtil.cpp b/Engine/Source/Runtime/TypedElementRuntime/Private/Elements/Framework/TypedElementListObjectUtil.cpp
--- a/Engine/Source/Runtime/TypedElementRuntime/Private/Elements/Framework/TypedElementListObjectUtil.cpp
+++ b/Engine/Source/Runtime/TypedElementRuntime/Private/Elements/Framework/TypedElementListObjectUtil.cpp
@@ -9,7 +9,7 @@ namespace TypedElementListObjectUtil
 UObject* GetObjectOfType(const TTypedElement<ITypedElementObjectInterface>& InObjectElement, const UClass* InRequiredClass)
 {
 	UObject* ElementObject = InObjectElement ? InObjectElement.GetObject() : nullptr;
-	return (ElementObject && (!InRequiredClass || ElementObject->GetClass()->ImplementsInterface(InRequiredClass)))
+	return (ElementObject && (!InRequiredClass || ElementObject->IsA(InRequiredClass) || ElementObject->GetClass()->ImplementsInterface(InRequiredClass)))
 		? ElementObject
 		: nullptr;
 }
