Fix a crash when editing project settings in typed element serialization code.

#rnx

diff --git a/Engine/Source/Runtime/Engine/Private/Elements/SMInstance/SMInstanceElementId.cpp b/Engine/Source/Runtime/Engine/Private/Elements/SMInstance/SMInstanceElementId.cpp
--- a/Engine/Source/Runtime/Engine/Private/Elements/SMInstance/SMInstanceElementId.cpp
+++ b/Engine/Source/Runtime/Engine/Private/Elements/SMInstance/SMInstanceElementId.cpp
@@ -6,6 +6,10 @@
 #include "UObject/Package.h"
 #include "Components/InstancedStaticMeshComponent.h"
 
+#if WITH_EDITORONLY_DATA
+#include "Serialization/TextReferenceCollector.h"
+#endif //WITH_EDITORONLY_DATA
+
 FSMInstanceElementIdMapEntry::FSMInstanceElementIdMapEntry(FSMInstanceElementIdMap* InOwner, UInstancedStaticMeshComponent* InComponent)
 	: Owner(InOwner)
 	, Component(InComponent)
@@ -355,6 +359,13 @@ void FSMInstanceElementIdMap::AddReferencedObjects(FReferenceCollector& Collecto
 #endif	// WITH_EDITORONLY_DATA
 }
 
+USMInstanceElementIdMapTransactor::USMInstanceElementIdMapTransactor()
+{
+#if WITH_EDITORONLY_DATA
+	// The map transactor should not get text reference collected, as it is native and transient only, and should not find its way into asset packages
+	{ static const FAutoRegisterTextReferenceCollectorCallback AutomaticRegistrationOfTextReferenceCollector(USMInstanceElementIdMapTransactor::StaticClass(), [](UObject* Object, FArchive& Ar) {}); }
+#endif //WITH_EDITORONLY_DATA
+}
 
 #if WITH_EDITORONLY_DATA
 void USMInstanceElementIdMapTransactor::Serialize(FArchive& Ar)
diff --git a/Engine/Source/Runtime/Engine/Public/Elements/SMInstance/SMInstanceElementId.h b/Engine/Source/Runtime/Engine/Public/Elements/SMInstance/SMInstanceElementId.h
--- a/Engine/Source/Runtime/Engine/Public/Elements/SMInstance/SMInstanceElementId.h
+++ b/Engine/Source/Runtime/Engine/Public/Elements/SMInstance/SMInstanceElementId.h
@@ -195,6 +195,8 @@ class USMInstanceElementIdMapTransactor : public UObject
 	GENERATED_BODY()
 
 public:
+	USMInstanceElementIdMapTransactor();
+
 #if WITH_EDITORONLY_DATA
 	void SetOwnerEntry(FSMInstanceElementIdMapEntry* InOwnerEntry)
 	{
diff --git a/Engine/Source/Runtime/TypedElementRuntime/Private/Elements/Framework/TypedElementSelectionSet.cpp b/Engine/Source/Runtime/TypedElementRuntime/Private/Elements/Framework/TypedElementSelectionSet.cpp
--- a/Engine/Source/Runtime/TypedElementRuntime/Private/Elements/Framework/TypedElementSelectionSet.cpp
+++ b/Engine/Source/Runtime/TypedElementRuntime/Private/Elements/Framework/TypedElementSelectionSet.cpp
@@ -10,6 +10,10 @@
 #include "Serialization/ObjectReader.h"
 #include "Serialization/ObjectWriter.h"
 
+#if WITH_EDITORONLY_DATA
+#include "Serialization/TextReferenceCollector.h"
+#endif //WITH_EDITORONLY_DATA
+
 void FTypedElementSelectionCustomization::GetNormalizedElements(const TTypedElement<ITypedElementSelectionInterface>& InElementSelectionHandle, FTypedElementListConstRef InSelectionSet, const FTypedElementSelectionNormalizationOptions& InNormalizationOptions, FTypedElementListRef OutNormalizedElements)
 {
 	// Don't include this element in the normalized selection if it has a parent element that is also selected
@@ -41,6 +45,11 @@ UTypedElementSelectionSet::UTypedElementSelectionSet()
 		ElementList->OnPreChange().AddUObject(this, &UTypedElementSelectionSet::OnElementListPreChange);
 		ElementList->OnChanged().AddUObject(this, &UTypedElementSelectionSet::OnElementListChanged);
 	}
+	
+#if WITH_EDITORONLY_DATA
+	// The selection set should not get text reference collected, as it is native and transient only, and should not find its way into asset packages
+	{ static const FAutoRegisterTextReferenceCollectorCallback AutomaticRegistrationOfTextReferenceCollector(UTypedElementSelectionSet::StaticClass(), [](UObject* Object, FArchive& Ar) {}); }
+#endif //WITH_EDITORONLY_DATA
 }
 
 #if WITH_EDITOR
