[Typed Elements] Fix an issue where the component visualizer was reset when components were reinstanced by bp actor post move

diff --git a/Engine/Source/Runtime/TypedElementRuntime/Private/Elements/Framework/TypedElementSelectionSet.cpp b/Engine/Source/Runtime/TypedElementRuntime/Private/Elements/Framework/TypedElementSelectionSet.cpp
--- a/Engine/Source/Runtime/TypedElementRuntime/Private/Elements/Framework/TypedElementSelectionSet.cpp
+++ b/Engine/Source/Runtime/TypedElementRuntime/Private/Elements/Framework/TypedElementSelectionSet.cpp
@@ -483,6 +483,11 @@ void UTypedElementSelectionSet::OnElementReplaced(TArrayView<const TTuple<FTyped
 		.SetAllowLegacyNotifications(false)
 		.SetWarnIfLocked(false);
 
+	// We're updating to keep in sync with element data that may have been changed by other aspects of the engine
+	// The selection isn't actually changing, so don't broadcast all the notifications about the select and deselects we do here.
+	FTypedElementList::FScopedClearNewPendingChange ScopedClearChanges = GetScopedClearNewPendingChange();
+	FTypedElementList::FLegacySyncScopedBatch LegacySync(*ElementList, SelectionOptions.AllowLegacyNotifications());
+
 	for (const TTuple<FTypedElementHandle, FTypedElementHandle>& ReplacedElement : InReplacedElements)
 	{
 		if (ElementList->Num() == 0)
@@ -522,7 +527,12 @@ void UTypedElementSelectionSet::OnElementUpdated(TArrayView<const FTypedElementH
 	{
 		return;
 	}
-	
+
+	// We're updating to keep in sync with element data that may have been changed by other aspects of the engine
+	// The selection isn't actually changing, so don't broadcast all the notifications about the select and deselects we do here.
+	FTypedElementList::FScopedClearNewPendingChange ScopedClearChanges = GetScopedClearNewPendingChange();
+	FTypedElementList::FLegacySyncScopedBatch LegacySync(*ElementList, SelectionOptions.AllowLegacyNotifications());
+
 	for (const FTypedElementHandle& UpdatedElement : InUpdatedElements)
 	{
 		FTypedElementSelectionSetElement SelectionSetElementToUpdate = ResolveSelectionSetElement(UpdatedElement);
