Fix compile error due to branch divergence caused by CL 20873839

#rnx

diff --git a/Engine/Source/Runtime/TypedElementRuntime/Private/Elements/Framework/TypedElementSelectionSet.cpp b/Engine/Source/Runtime/TypedElementRuntime/Private/Elements/Framework/TypedElementSelectionSet.cpp
--- a/Engine/Source/Runtime/TypedElementRuntime/Private/Elements/Framework/TypedElementSelectionSet.cpp
+++ b/Engine/Source/Runtime/TypedElementRuntime/Private/Elements/Framework/TypedElementSelectionSet.cpp
@@ -486,7 +486,7 @@ void UTypedElementSelectionSet::OnElementReplaced(TArrayView<const TTuple<FTyped
 	// We're updating to keep in sync with element data that may have been changed by other aspects of the engine
 	// The selection isn't actually changing, so don't broadcast all the notifications about the select and deselects we do here.
 	FTypedElementList::FScopedClearNewPendingChange ScopedClearChanges = GetScopedClearNewPendingChange();
-	FTypedElementList::FLegacySyncScopedBatch LegacySync(*ElementList, SelectionOptions.AllowLegacyNotifications());
+	FTypedElementListLegacySyncScopedBatch LegacySync(*ElementList, SelectionOptions.AllowLegacyNotifications());
 
 	for (const TTuple<FTypedElementHandle, FTypedElementHandle>& ReplacedElement : InReplacedElements)
 	{
@@ -531,7 +531,7 @@ void UTypedElementSelectionSet::OnElementUpdated(TArrayView<const FTypedElementH
 	// We're updating to keep in sync with element data that may have been changed by other aspects of the engine
 	// The selection isn't actually changing, so don't broadcast all the notifications about the select and deselects we do here.
 	FTypedElementList::FScopedClearNewPendingChange ScopedClearChanges = GetScopedClearNewPendingChange();
-	FTypedElementList::FLegacySyncScopedBatch LegacySync(*ElementList, SelectionOptions.AllowLegacyNotifications());
+	FTypedElementListLegacySyncScopedBatch LegacySync(*ElementList, SelectionOptions.AllowLegacyNotifications());
 
 	for (const FTypedElementHandle& UpdatedElement : InUpdatedElements)
 	{
