TypedElementSelectionSet: Make sure child selection changes are reflected by the returned value

diff --git a/Engine/Source/Runtime/TypedElementRuntime/Private/Elements/Framework/TypedElementSelectionSet.cpp b/Engine/Source/Runtime/TypedElementRuntime/Private/Elements/Framework/TypedElementSelectionSet.cpp
--- a/Engine/Source/Runtime/TypedElementRuntime/Private/Elements/Framework/TypedElementSelectionSet.cpp
+++ b/Engine/Source/Runtime/TypedElementRuntime/Private/Elements/Framework/TypedElementSelectionSet.cpp
@@ -210,7 +210,7 @@ bool UTypedElementSelectionSet::SelectElement(const FTypedElementHandle& InEleme
 			ElementHierarchyHandle.GetChildElements(ChildElementHandles, /*bAllowCreate*/true);
 			for (const FTypedElementHandle& ChildElementHandle : ChildElementHandles)
 			{
-				SelectElement(ChildElementHandle, ChildSelectionOptions);
+				bSelectionChanged |= SelectElement(ChildElementHandle, ChildSelectionOptions);
 			}
 		}
 	}
@@ -277,7 +277,7 @@ bool UTypedElementSelectionSet::DeselectElement(const FTypedElementHandle& InEle
 			ElementHierarchyHandle.GetChildElements(ChildElementHandles, /*bAllowCreate*/false);
 			for (const FTypedElementHandle& ChildElementHandle : ChildElementHandles)
 			{
-				DeselectElement(ChildElementHandle, ChildSelectionOptions);
+				bSelectionChanged |= DeselectElement(ChildElementHandle, ChildSelectionOptions);
 			}
 		}
 	}
