Add Primitive CustomData Interface for SM Instances
* Allows editor utility blueprints to set the CustomData primitive data on selected elements through the TypedElement system

diff --git a/Engine/Source/Runtime/Engine/Private/Elements/Framework/EngineElements.cpp b/Engine/Source/Runtime/Engine/Private/Elements/Framework/EngineElements.cpp
--- a/Engine/Source/Runtime/Engine/Private/Elements/Framework/EngineElements.cpp
+++ b/Engine/Source/Runtime/Engine/Private/Elements/Framework/EngineElements.cpp
@@ -29,6 +29,7 @@
 #include "Elements/SMInstance/SMInstanceElementWorldInterface.h"
 #include "Elements/SMInstance/SMInstanceElementSelectionInterface.h"
 #include "Elements/SMInstance/SMInstanceElementAssetDataInterface.h"
+#include "Elements/SMInstance/SMInstanceElementPrimitiveCustomDataInterface.h"
 
 #include "Modules/ModuleManager.h"
 
@@ -79,6 +80,7 @@ void RegisterEngineSMInstanceElements()
 	Registry->RegisterElementInterface<ITypedElementWorldInterface>(NAME_SMInstance, NewObject<USMInstanceElementWorldInterface>());
 	Registry->RegisterElementInterface<ITypedElementSelectionInterface>(NAME_SMInstance, NewObject<USMInstanceElementSelectionInterface>());
 	Registry->RegisterElementInterface<ITypedElementAssetDataInterface>(NAME_SMInstance, NewObject<USMInstanceElementAssetDataInterface>());
+	Registry->RegisterElementInterface<ITypedElementPrimitiveCustomDataInterface>(NAME_SMInstance, NewObject<USMInstanceElementPrimitiveCustomDataInterface>());
 }
 
 void RegisterEngineElements()
diff --git a/Engine/Source/Runtime/Engine/Private/Elements/SMInstance/SMInstanceElementPrimitiveCustomDataInterface.cpp b/Engine/Source/Runtime/Engine/Private/Elements/SMInstance/SMInstanceElementPrimitiveCustomDataInterface.cpp
new file mode 100644
--- /dev/null
+++ b/Engine/Source/Runtime/Engine/Private/Elements/SMInstance/SMInstanceElementPrimitiveCustomDataInterface.cpp
@@ -0,0 +1,28 @@
+﻿// Copyright Epic Games, Inc. All Rights Reserved.
+
+#include "Elements/SMInstance/SMInstanceElementPrimitiveCustomDataInterface.h"
+
+#include "Components/InstancedStaticMeshComponent.h"
+#include "Elements/SMInstance/SMInstanceElementData.h"
+
+#include UE_INLINE_GENERATED_CPP_BY_NAME(SMInstanceElementPrimitiveCustomDataInterface)
+
+void USMInstanceElementPrimitiveCustomDataInterface::SetCustomData(const FTypedElementHandle& InElementHandle,
+	TArrayView<const float> CustomDataFloats, bool bMarkRenderStateDirty)
+{
+	FSMInstanceManager SMInstance = SMInstanceElementDataUtil::GetSMInstanceFromHandle(InElementHandle);
+	if (SMInstance)
+	{
+		SMInstance.GetISMComponent()->SetCustomData(SMInstance.GetISMInstanceIndex(), CustomDataFloats, bMarkRenderStateDirty);
+	}
+}
+
+void USMInstanceElementPrimitiveCustomDataInterface::SetCustomDataValue(const FTypedElementHandle& InElementHandle,
+	int32 CustomDataIndex, float CustomDataValue, bool bMarkRenderStateDirty)
+{
+	FSMInstanceManager SMInstance = SMInstanceElementDataUtil::GetSMInstanceFromHandle(InElementHandle);
+	if (SMInstance)
+	{
+		SMInstance.GetISMComponent()->SetCustomDataValue(SMInstance.GetISMInstanceIndex(), CustomDataIndex, CustomDataValue, bMarkRenderStateDirty);
+	}
+}
diff --git a/Engine/Source/Runtime/Engine/Public/Elements/SMInstance/SMInstanceElementPrimitiveCustomDataInterface.h b/Engine/Source/Runtime/Engine/Public/Elements/SMInstance/SMInstanceElementPrimitiveCustomDataInterface.h
new file mode 100644
--- /dev/null
+++ b/Engine/Source/Runtime/Engine/Public/Elements/SMInstance/SMInstanceElementPrimitiveCustomDataInterface.h
@@ -0,0 +1,16 @@
+﻿// Copyright Epic Games, Inc. All Rights Reserved.
+
+#pragma once
+
+#include "Elements/Interfaces/TypedElementPrimitiveCustomDataInterface.h"
+#include "SMInstanceElementPrimitiveCustomDataInterface.generated.h"
+
+UCLASS()
+class ENGINE_API USMInstanceElementPrimitiveCustomDataInterface : public UObject, public ITypedElementPrimitiveCustomDataInterface
+{
+	GENERATED_BODY()
+
+public:
+	virtual void SetCustomData(const FTypedElementHandle& InElementHandle, TArrayView<const float> CustomDataFloats,  bool bMarkRenderStateDirty) override;
+	virtual void SetCustomDataValue(const FTypedElementHandle& InElementHandle, int32 CustomDataIndex, float CustomDataValue, bool bMarkRenderStateDirty) override;
+};
diff --git a/Engine/Source/Runtime/TypedElementRuntime/Private/Elements/Interfaces/TypedElementPrimitiveCustomDataInterface.cpp b/Engine/Source/Runtime/TypedElementRuntime/Private/Elements/Interfaces/TypedElementPrimitiveCustomDataInterface.cpp
new file mode 100644
--- /dev/null
+++ b/Engine/Source/Runtime/TypedElementRuntime/Private/Elements/Interfaces/TypedElementPrimitiveCustomDataInterface.cpp
@@ -0,0 +1,30 @@
+﻿// Copyright Epic Games, Inc. All Rights Reserved.
+
+#include "Elements/Interfaces/TypedElementPrimitiveCustomDataInterface.h"
+
+#include UE_INLINE_GENERATED_CPP_BY_NAME(TypedElementPrimitiveCustomDataInterface)
+
+void ITypedElementPrimitiveCustomDataInterface::SetCustomData(const FScriptTypedElementHandle& InElementHandle, const TArray<float>& CustomDataFloats,  bool bMarkRenderStateDirty)
+{
+	FTypedElementHandle NativeHandle = InElementHandle.GetTypedElementHandle();
+	if (!NativeHandle)
+	{
+		FFrame::KismetExecutionMessage(TEXT("InElementHandle is not a valid handle."), ELogVerbosity::Error);
+		return;
+	}
+
+	SetCustomData(NativeHandle, MakeArrayView(CustomDataFloats.GetData(), CustomDataFloats.Num()), bMarkRenderStateDirty);
+	
+}
+
+void ITypedElementPrimitiveCustomDataInterface::SetCustomDataValue(const FScriptTypedElementHandle& InElementHandle, int32 CustomDataIndex, float CustomDataValue, bool bMarkRenderStateDirty)
+{
+	FTypedElementHandle NativeHandle = InElementHandle.GetTypedElementHandle();
+	if (!NativeHandle)
+	{
+		FFrame::KismetExecutionMessage(TEXT("InElementHandle is not a valid handle."), ELogVerbosity::Error);
+		return;
+	}
+	
+	SetCustomDataValue(NativeHandle, CustomDataIndex, CustomDataValue, bMarkRenderStateDirty);
+}
\ No newline at end of file
diff --git a/Engine/Source/Runtime/TypedElementRuntime/Public/Elements/Interfaces/TypedElementPrimitiveCustomDataInterface.h b/Engine/Source/Runtime/TypedElementRuntime/Public/Elements/Interfaces/TypedElementPrimitiveCustomDataInterface.h
new file mode 100644
--- /dev/null
+++ b/Engine/Source/Runtime/TypedElementRuntime/Public/Elements/Interfaces/TypedElementPrimitiveCustomDataInterface.h
@@ -0,0 +1,45 @@
+﻿// Copyright Epic Games, Inc. All Rights Reserved.
+
+#pragma once
+
+#include "Elements/Framework/TypedElementHandle.h"
+#include "UObject/Interface.h"
+
+#include "TypedElementPrimitiveCustomDataInterface.generated.h"
+
+UINTERFACE(MinimalAPI, BlueprintType, meta = (CannotImplementInterfaceInBlueprint))
+class UTypedElementPrimitiveCustomDataInterface : public UInterface
+{
+	GENERATED_BODY()
+};
+
+/**
+ * An interface for use with the TypedElement Framework which exposes Primitive CustomData
+ * for use within Materials.
+ */
+class TYPEDELEMENTRUNTIME_API ITypedElementPrimitiveCustomDataInterface
+{
+	GENERATED_BODY()
+
+public:
+	
+	virtual void SetCustomData(const FTypedElementHandle& InElementHandle, TArrayView<const float> CustomDataFloats,  bool bMarkRenderStateDirty = false) {}
+	virtual void SetCustomDataValue(const FTypedElementHandle& InElementHandle, int32 CustomDataIndex, float CustomDataValue, bool bMarkRenderStateDirty = false) {};
+
+	//~ Scripting Interface
+
+	// Sets all Primitive's CustomData values
+	UFUNCTION(BlueprintCallable, Category="TypedElementInterfaces|CustomData")
+	virtual void SetCustomData(const FScriptTypedElementHandle& InElementHandle, const TArray<float>& CustomDataFloats,  bool bMarkRenderStateDirty = false);
+
+	// Sets a single Primitive's CustomData value
+	UFUNCTION(BlueprintCallable, Category="TypedElementInterfaces|CustomData")
+	virtual void SetCustomDataValue(const FScriptTypedElementHandle& InElementHandle, int32 CustomDataIndex, float CustomDataValue, bool bMarkRenderStateDirty = false);
+};
+
+template <>
+struct TTypedElement<ITypedElementPrimitiveCustomDataInterface> : public TTypedElementBase<ITypedElementPrimitiveCustomDataInterface>
+{
+	void SetCustomData(TArrayView<float> CustomDataFloats,  bool bMarkRenderStateDirty = false) const { return InterfacePtr->SetCustomData(*this, CustomDataFloats, bMarkRenderStateDirty); }
+	void SetCustomDataValue(int32 CustomDataIndex, float CustomDataValue, bool bMarkRenderStateDirty = false) { return InterfacePtr->SetCustomDataValue(*this, CustomDataIndex, CustomDataValue, bMarkRenderStateDirty); }
+};
\ No newline at end of file
