[Engine]

diff --git a/Engine/Source/Runtime/TypedElementFramework/Public/Elements/Framework/TypedElementRegistry.h b/Engine/Source/Runtime/TypedElementFramework/Public/Elements/Framework/TypedElementRegistry.h
--- a/Engine/Source/Runtime/TypedElementFramework/Public/Elements/Framework/TypedElementRegistry.h
+++ b/Engine/Source/Runtime/TypedElementFramework/Public/Elements/Framework/TypedElementRegistry.h
@@ -50,47 +50,47 @@
 /**
  * Registry of element types and their associated interfaces, along with the elements that represent their instances.
  */
-UCLASS(Transient)
-class TYPEDELEMENTFRAMEWORK_API UTypedElementRegistry : public UObject
+UCLASS(Transient, MinimalAPI)
+class UTypedElementRegistry : public UObject
 {
 	GENERATED_BODY()
 
 public:
-	UTypedElementRegistry();
+	TYPEDELEMENTFRAMEWORK_API UTypedElementRegistry();
 
 	//~ UObject interface
-	virtual void FinishDestroy() override;
-	static void AddReferencedObjects(UObject* InThis, FReferenceCollector& Collector);
+	TYPEDELEMENTFRAMEWORK_API virtual void FinishDestroy() override;
+	static TYPEDELEMENTFRAMEWORK_API void AddReferencedObjects(UObject* InThis, FReferenceCollector& Collector);
 
 	/**
 	 * Initialize the singleton instance of the registry used in most cases.
 	 */
-	static void Private_InitializeInstance();
+	static TYPEDELEMENTFRAMEWORK_API void Private_InitializeInstance();
 
 	/**
 	 * Shutdown the singleton instance of the registry used in most cases.
 	 */
-	static void Private_ShutdownInstance();
+	static TYPEDELEMENTFRAMEWORK_API void Private_ShutdownInstance();
 
 	/**
 	 * Get the singleton instance of the registry used in most cases.
 	 */
 	UFUNCTION(BlueprintPure, DisplayName="Get Default Typed Element Registry", Category = "TypedElementFramework|Registry", meta=(ScriptName="GetDefaultTypedElementRegistry"))
-	static UTypedElementRegistry* GetInstance();
+	static TYPEDELEMENTFRAMEWORK_API UTypedElementRegistry* GetInstance();
 
-	ITypedElementDataStorageInterface* GetMutableDataStorage();
-	const ITypedElementDataStorageInterface* GetDataStorage() const;
-	void SetDataStorage(ITypedElementDataStorageInterface* Storage);
+	TYPEDELEMENTFRAMEWORK_API ITypedElementDataStorageInterface* GetMutableDataStorage();
+	TYPEDELEMENTFRAMEWORK_API const ITypedElementDataStorageInterface* GetDataStorage() const;
+	TYPEDELEMENTFRAMEWORK_API void SetDataStorage(ITypedElementDataStorageInterface* Storage);
 
-	ITypedElementDataStorageCompatibilityInterface* GetMutableDataStorageCompatibility();
-	const ITypedElementDataStorageCompatibilityInterface* GetDataStorageCompatibility() const;
-	void SetDataStorageCompatibility(ITypedElementDataStorageCompatibilityInterface* Storage);
+	TYPEDELEMENTFRAMEWORK_API ITypedElementDataStorageCompatibilityInterface* GetMutableDataStorageCompatibility();
+	TYPEDELEMENTFRAMEWORK_API const ITypedElementDataStorageCompatibilityInterface* GetDataStorageCompatibility() const;
+	TYPEDELEMENTFRAMEWORK_API void SetDataStorageCompatibility(ITypedElementDataStorageCompatibilityInterface* Storage);
 
-	ITypedElementDataStorageUiInterface* GetMutableDataStorageUi();
-	const ITypedElementDataStorageUiInterface* GetDataStorageUi() const;
-	void SetDataStorageUi(ITypedElementDataStorageUiInterface* Storage);
+	TYPEDELEMENTFRAMEWORK_API ITypedElementDataStorageUiInterface* GetMutableDataStorageUi();
+	TYPEDELEMENTFRAMEWORK_API const ITypedElementDataStorageUiInterface* GetDataStorageUi() const;
+	TYPEDELEMENTFRAMEWORK_API void SetDataStorageUi(ITypedElementDataStorageUiInterface* Storage);
 
-	bool AreDataStorageInterfacesSet() const;
+	TYPEDELEMENTFRAMEWORK_API bool AreDataStorageInterfacesSet() const;
 
 	/**
 	 * Event fired when all Data Storage Interfaces have been set. 
@@ -263,17 +263,17 @@ public:
 	 * @note This is automatically called at the end of each frame, but may also be called manually.
 	 * @note This is automatically called post-GC, unless auto-GC destruction has been disabled (@see FDisableElementDestructionOnGC).
 	 */
-	void ProcessDeferredElementsToDestroy();
+	TYPEDELEMENTFRAMEWORK_API void ProcessDeferredElementsToDestroy();
 
 	/**
 	 * Release an element ID that was previously acquired from an existing handle.
 	 */
-	void ReleaseElementId(FTypedElementId& InOutElementId);
+	TYPEDELEMENTFRAMEWORK_API void ReleaseElementId(FTypedElementId& InOutElementId);
 
 	/**
 	 * Get an element handle from its minimal ID.
 	 */
-	FTypedElementHandle GetElementHandle(const FTypedElementId& InElementId) const;
+	TYPEDELEMENTFRAMEWORK_API FTypedElementHandle GetElementHandle(const FTypedElementId& InElementId) const;
 
 	/**
 	 * Get an element that implements the given interface from its minimal ID.
@@ -336,12 +336,12 @@ public:
 	/**
 	 * Create an empty list of elements associated with this registry, populated from the given minimal IDs that are valid.
 	 */
-	FTypedElementListRef CreateElementList(TArrayView<const FTypedElementId> InElementIds);
+	TYPEDELEMENTFRAMEWORK_API FTypedElementListRef CreateElementList(TArrayView<const FTypedElementId> InElementIds);
 
 	/**
 	 * Create an empty list of elements associated with this registry, populated from the given handles that are valid.
 	 */
-	FTypedElementListRef CreateElementList(TArrayView<const FTypedElementHandle> InElementHandles);
+	TYPEDELEMENTFRAMEWORK_API FTypedElementListRef CreateElementList(TArrayView<const FTypedElementHandle> InElementHandles);
 
 	/**
 	 * Create an empty list of elements associated with this registry, populated from the given owners that are valid.
@@ -367,7 +367,7 @@ public:
 	 * Extract a string that contains name of the type elements registered.
 	 * For each type it also says what interface are registered and the path of the class that implement them
 	 */
-	FString RegistredElementTypesAndInterfacesToString() const;
+	TYPEDELEMENTFRAMEWORK_API FString RegistredElementTypesAndInterfacesToString() const;
 
 	/**
 	 * Create a script handle. This should only be use for the script exposure apis.
@@ -548,9 +548,9 @@ private:
 		TArray<FDeferredElementToRemove> DeferredElementsToRemove;
 	};
 
-	void RegisterElementTypeImpl(const FName InElementTypeName, TUniquePtr<FRegisteredElementType>&& InRegisteredElementType);
-	void RegisterElementInterfaceImpl(const FName InElementTypeName, UObject* InElementInterface, const TSubclassOf<UInterface>& InBaseInterfaceType, const bool InAllowOverride);
-	UObject* GetElementInterfaceImpl(const FTypedHandleTypeId InElementTypeId, const TSubclassOf<UInterface>& InBaseInterfaceType) const;
+	TYPEDELEMENTFRAMEWORK_API void RegisterElementTypeImpl(const FName InElementTypeName, TUniquePtr<FRegisteredElementType>&& InRegisteredElementType);
+	TYPEDELEMENTFRAMEWORK_API void RegisterElementInterfaceImpl(const FName InElementTypeName, UObject* InElementInterface, const TSubclassOf<UInterface>& InBaseInterfaceType, const bool InAllowOverride);
+	TYPEDELEMENTFRAMEWORK_API UObject* GetElementInterfaceImpl(const FTypedHandleTypeId InElementTypeId, const TSubclassOf<UInterface>& InBaseInterfaceType) const;
 
 	template <typename ElementDataType>
 	TTypedElementOwner<ElementDataType> CreateElementImpl(const FName InElementTypeName, const FTypedHandleElementId InElementId)
@@ -665,11 +665,11 @@ private:
 		return nullptr;
 	}
 
-	void NotifyElementListPendingChanges();
+	TYPEDELEMENTFRAMEWORK_API void NotifyElementListPendingChanges();
 
-	void OnBeginFrame();
-	void OnEndFrame();
-	void OnPostGarbageCollect();
+	TYPEDELEMENTFRAMEWORK_API void OnBeginFrame();
+	TYPEDELEMENTFRAMEWORK_API void OnEndFrame();
+	TYPEDELEMENTFRAMEWORK_API void OnPostGarbageCollect();
 
 	FORCEINLINE void IncrementDisableElementDestructionOnGCCount()
 	{
@@ -683,7 +683,7 @@ private:
 		--DisableElementDestructionOnGCCount;
 	}
 
-	void CallDataStorageInterfacesSetDelegateIfNeeded();
+	TYPEDELEMENTFRAMEWORK_API void CallDataStorageInterfacesSetDelegateIfNeeded();
 
 	mutable FRWLock RegisteredElementTypesRW;
 	TUniquePtr<FRegisteredElementType> RegisteredElementTypes[TypedHandleMaxTypeId - 1];
diff --git a/Engine/Source/Runtime/TypedElementFramework/Public/Elements/Interfaces/TypedElementDataStorageCompatibilityInterface.h b/Engine/Source/Runtime/TypedElementFramework/Public/Elements/Interfaces/TypedElementDataStorageCompatibilityInterface.h
--- a/Engine/Source/Runtime/TypedElementFramework/Public/Elements/Interfaces/TypedElementDataStorageCompatibilityInterface.h
+++ b/Engine/Source/Runtime/TypedElementFramework/Public/Elements/Interfaces/TypedElementDataStorageCompatibilityInterface.h
@@ -24,7 +24,7 @@ class UTypedElementDataStorageCompatibilityInterface : public UInterface
  * Interface to provide compatibility with existing systems that don't directly
  * support the data storage.
  */
-class TYPEDELEMENTFRAMEWORK_API ITypedElementDataStorageCompatibilityInterface
+class ITypedElementDataStorageCompatibilityInterface
 {
 	GENERATED_BODY()
 
diff --git a/Engine/Source/Runtime/TypedElementFramework/Public/Elements/Interfaces/TypedElementDataStorageFactory.h b/Engine/Source/Runtime/TypedElementFramework/Public/Elements/Interfaces/TypedElementDataStorageFactory.h
--- a/Engine/Source/Runtime/TypedElementFramework/Public/Elements/Interfaces/TypedElementDataStorageFactory.h
+++ b/Engine/Source/Runtime/TypedElementFramework/Public/Elements/Interfaces/TypedElementDataStorageFactory.h
@@ -11,8 +11,8 @@ class ITypedElementDataStorageUiInterface;
  * Base class that can be used to register various elements, such as queries and widgets, with
  * the Typed Elements Data Storage.
  */
-UCLASS()
-class TYPEDELEMENTFRAMEWORK_API UTypedElementDataStorageFactory : public UObject
+UCLASS(MinimalAPI)
+class UTypedElementDataStorageFactory : public UObject
 {
 	GENERATED_BODY()
 
@@ -32,4 +32,4 @@ public:
 	virtual void RegisterWidgetPurposes(ITypedElementDataStorageUiInterface& DataStorageUi) const {}
 	virtual void RegisterWidgetConstructors(ITypedElementDataStorageInterface& DataStorage,
 		ITypedElementDataStorageUiInterface& DataStorageUi) const {}
-};
\ No newline at end of file
+};
diff --git a/Engine/Source/Runtime/TypedElementFramework/Public/Elements/Interfaces/TypedElementDataStorageInterface.h b/Engine/Source/Runtime/TypedElementFramework/Public/Elements/Interfaces/TypedElementDataStorageInterface.h
--- a/Engine/Source/Runtime/TypedElementFramework/Public/Elements/Interfaces/TypedElementDataStorageInterface.h
+++ b/Engine/Source/Runtime/TypedElementFramework/Public/Elements/Interfaces/TypedElementDataStorageInterface.h
@@ -82,7 +82,7 @@ struct TTypedElementColumnTypeList
 	operator TConstArrayView<const UScriptStruct*>() const { return ColumnTypes; }
 };
 
-class TYPEDELEMENTFRAMEWORK_API ITypedElementDataStorageInterface
+class ITypedElementDataStorageInterface
 {
 	GENERATED_BODY()
 
diff --git a/Engine/Source/Runtime/TypedElementFramework/Public/Elements/Interfaces/TypedElementDataStorageUiInterface.h b/Engine/Source/Runtime/TypedElementFramework/Public/Elements/Interfaces/TypedElementDataStorageUiInterface.h
--- a/Engine/Source/Runtime/TypedElementFramework/Public/Elements/Interfaces/TypedElementDataStorageUiInterface.h
+++ b/Engine/Source/Runtime/TypedElementFramework/Public/Elements/Interfaces/TypedElementDataStorageUiInterface.h
@@ -24,21 +24,21 @@ class SWidget;
  * arguments can be safely applied.
  */
 USTRUCT()
-struct TYPEDELEMENTFRAMEWORK_API FTypedElementWidgetConstructor
+struct FTypedElementWidgetConstructor
 {
 	GENERATED_BODY()
 
 public:
-	explicit FTypedElementWidgetConstructor(const UScriptStruct* InTypeInfo);
+	TYPEDELEMENTFRAMEWORK_API explicit FTypedElementWidgetConstructor(const UScriptStruct* InTypeInfo);
 	explicit FTypedElementWidgetConstructor(EForceInit) {} //< For compatibility and shouldn't be directly used.
 
 	virtual ~FTypedElementWidgetConstructor() = default;
 
 	/** Retrieves the type information for the constructor type. */
-	virtual const UScriptStruct* GetTypeInfo() const;
+	TYPEDELEMENTFRAMEWORK_API virtual const UScriptStruct* GetTypeInfo() const;
 
 	/** Returns a list of additional columns the widget requires. */
-	virtual TConstArrayView<const UScriptStruct*> GetAdditionalColumnsList() const;
+	TYPEDELEMENTFRAMEWORK_API virtual TConstArrayView<const UScriptStruct*> GetAdditionalColumnsList() const;
 
 	/** 
 	 * Whether or not an instance of this constructor can be reused. Setting this to true means that the Data Storage will
@@ -46,7 +46,7 @@ public:
 	 * 'false'. Override this function and set it to true if there's no internal state that could interfere with creating
 	 * multiple widgets from the same constructor, e.g. data stored from passed in arguments.
 	 */
-	virtual bool CanBeReused() const;
+	TYPEDELEMENTFRAMEWORK_API virtual bool CanBeReused() const;
 
 	/**
 	 * Constructs the widget according to the provided information. Information is collected by calling
@@ -54,7 +54,7 @@ public:
 	 * functions to build widgets according to a standard recipe and to reduce the amount of code needed.
 	 * If a complexer situation is called for this function can also be directly overwritten.
 	 */
-	virtual TSharedPtr<SWidget> Construct(
+	TYPEDELEMENTFRAMEWORK_API virtual TSharedPtr<SWidget> Construct(
 		TypedElementRowHandle Row, /** The row the widget will be stored in. */
 		ITypedElementDataStorageInterface* DataStorage,
 		ITypedElementDataStorageUiInterface* DataStorageUi,
@@ -66,16 +66,16 @@ protected:
 	 * function if there are non-properties that can be configured through the arguments or if the configuration
 	 * contains arguments that are too complex to initialize (fully) through the type system.
 	 */
-	virtual bool ApplyArguments(TConstArrayView<TypedElement::ColumnUtils::Argument> Arguments);
+	TYPEDELEMENTFRAMEWORK_API virtual bool ApplyArguments(TConstArrayView<TypedElement::ColumnUtils::Argument> Arguments);
 	/** Create a new instance of the target widget. This is a required function. */
-	virtual TSharedPtr<SWidget> CreateWidget() PURE_VIRTUAL(FTypedElementWidgetConstructor::CreateWidget, return nullptr; );
+	TYPEDELEMENTFRAMEWORK_API virtual TSharedPtr<SWidget> CreateWidget() PURE_VIRTUAL(FTypedElementWidgetConstructor::CreateWidget, return nullptr; );
 	/** Set any values in columns if needed. The columns provided through GetAdditionalColumnsList() will have already been created. */
-	virtual bool SetColumns(ITypedElementDataStorageInterface* DataStorage, TypedElementRowHandle Row);
+	TYPEDELEMENTFRAMEWORK_API virtual bool SetColumns(ITypedElementDataStorageInterface* DataStorage, TypedElementRowHandle Row);
 	/** 
 	 * Last opportunity to configure anything in the widget or the row. This step can be needed to initialize widgets with data stored
 	 * in columns.
 	 */
-	virtual bool FinalizeWidget(
+	TYPEDELEMENTFRAMEWORK_API virtual bool FinalizeWidget(
 		ITypedElementDataStorageInterface* DataStorage,
 		ITypedElementDataStorageUiInterface* DataStorageUi,
 		TypedElementRowHandle Row,
@@ -100,7 +100,7 @@ class UTypedElementDataStorageUiInterface : public UInterface
 	GENERATED_BODY()
 };
 
-class TYPEDELEMENTFRAMEWORK_API ITypedElementDataStorageUiInterface
+class ITypedElementDataStorageUiInterface
 {
 	GENERATED_BODY()
 
diff --git a/Engine/Source/Runtime/TypedElementRuntime/Public/Elements/Framework/TypedElementInterfaceCustomization.h b/Engine/Source/Runtime/TypedElementRuntime/Public/Elements/Framework/TypedElementInterfaceCustomization.h
--- a/Engine/Source/Runtime/TypedElementRuntime/Public/Elements/Framework/TypedElementInterfaceCustomization.h
+++ b/Engine/Source/Runtime/TypedElementRuntime/Public/Elements/Framework/TypedElementInterfaceCustomization.h
@@ -13,7 +13,7 @@
 /**
  * Non-templated base class for the interface customization registry.
  */
-class TYPEDELEMENTRUNTIME_API FTypedElementInterfaceCustomizationRegistryBase
+class FTypedElementInterfaceCustomizationRegistryBase
 {
 public:
 	virtual ~FTypedElementInterfaceCustomizationRegistryBase() = default;
@@ -23,13 +23,13 @@ protected:
 	 * Given an element name, attempt to get its registered type ID from the global registry.
 	 * @return The registered type ID, or 0 if the element name is not registered.
 	 */
-	FTypedHandleTypeId GetElementTypeIdFromName(const FName InElementTypeName) const;
+	TYPEDELEMENTRUNTIME_API FTypedHandleTypeId GetElementTypeIdFromName(const FName InElementTypeName) const;
 
 	/**
 	 * Given an element name, attempt to get its registered type ID from the global registry.
 	 * @return The registered type ID, or asserts if the element name is not registered.
 	 */
-	FTypedHandleTypeId GetElementTypeIdFromNameChecked(const FName InElementTypeName) const;
+	TYPEDELEMENTRUNTIME_API FTypedHandleTypeId GetElementTypeIdFromNameChecked(const FName InElementTypeName) const;
 };
 
 /**
diff --git a/Engine/Source/Runtime/TypedElementRuntime/Public/Elements/Framework/TypedElementSelectionSet.h b/Engine/Source/Runtime/TypedElementRuntime/Public/Elements/Framework/TypedElementSelectionSet.h
--- a/Engine/Source/Runtime/TypedElementRuntime/Public/Elements/Framework/TypedElementSelectionSet.h
+++ b/Engine/Source/Runtime/TypedElementRuntime/Public/Elements/Framework/TypedElementSelectionSet.h
@@ -56,7 +56,7 @@ private:
  * Customization type used to allow asset editors (such as the level editor) to override the base behavior of element selection,
  * by injecting extra pre/post selection logic around the call into the selection interface for an element type.
  */
-class TYPEDELEMENTRUNTIME_API FTypedElementSelectionCustomization
+class FTypedElementSelectionCustomization
 {
 public:
 	virtual ~FTypedElementSelectionCustomization() = default;
@@ -69,13 +69,13 @@ public:
 	virtual bool DeselectElement(const TTypedElement<ITypedElementSelectionInterface>& InElementSelectionHandle, FTypedElementListRef InSelectionSet, const FTypedElementSelectionOptions& InSelectionOptions) { return InElementSelectionHandle.DeselectElement(InSelectionSet, InSelectionOptions); }
 	virtual bool AllowSelectionModifiers(const TTypedElement<ITypedElementSelectionInterface>& InElementSelectionHandle, FTypedElementListConstRef InSelectionSet) { return InElementSelectionHandle.AllowSelectionModifiers(InSelectionSet); }
 	virtual FTypedElementHandle GetSelectionElement(const TTypedElement<ITypedElementSelectionInterface>& InElementSelectionHandle, FTypedElementListConstRef InCurrentSelection, const ETypedElementSelectionMethod InSelectionMethod) { return InElementSelectionHandle.GetSelectionElement(InCurrentSelection, InSelectionMethod); }
-	virtual void GetNormalizedElements(const TTypedElement<ITypedElementSelectionInterface>& InElementSelectionHandle, FTypedElementListConstRef InSelectionSet, const FTypedElementSelectionNormalizationOptions& InNormalizationOptions, FTypedElementListRef OutNormalizedElements);
+	TYPEDELEMENTRUNTIME_API virtual void GetNormalizedElements(const TTypedElement<ITypedElementSelectionInterface>& InElementSelectionHandle, FTypedElementListConstRef InSelectionSet, const FTypedElementSelectionNormalizationOptions& InNormalizationOptions, FTypedElementListRef OutNormalizedElements);
 };
 
 /**
  * Utility to hold a typed element handle and its associated selection interface and selection customization.
  */
-struct TYPEDELEMENTRUNTIME_API FTypedElementSelectionSetElement
+struct FTypedElementSelectionSetElement
 {
 public:
 	FTypedElementSelectionSetElement() = default;
@@ -169,105 +169,105 @@ private:
  * A wrapper around an element list that ensures mutation goes via the selection 
  * interfaces, as well as providing some utilities for batching operations.
  */
-UCLASS(Transient, BlueprintType, meta=(DontUseGenericSpawnObject="True"))
-class TYPEDELEMENTRUNTIME_API UTypedElementSelectionSet : public UObject, public TTypedElementInterfaceCustomizationRegistry<FTypedElementSelectionCustomization>
+UCLASS(Transient, BlueprintType, meta=(DontUseGenericSpawnObject="True"), MinimalAPI)
+class UTypedElementSelectionSet : public UObject, public TTypedElementInterfaceCustomizationRegistry<FTypedElementSelectionCustomization>
 {
 	GENERATED_BODY()
 
 public:
-	UTypedElementSelectionSet();
+	TYPEDELEMENTRUNTIME_API UTypedElementSelectionSet();
 
 	//~ UObject interface
 #if WITH_EDITOR
-	virtual void PreEditUndo() override;
-	virtual void PostEditUndo() override;
-	virtual bool Modify(bool bAlwaysMarkDirty = true) override;
+	TYPEDELEMENTRUNTIME_API virtual void PreEditUndo() override;
+	TYPEDELEMENTRUNTIME_API virtual void PostEditUndo() override;
+	TYPEDELEMENTRUNTIME_API virtual bool Modify(bool bAlwaysMarkDirty = true) override;
 #endif	// WITH_EDITOR
-	virtual void BeginDestroy() override;
-	virtual void Serialize(FArchive& Ar) override;
+	TYPEDELEMENTRUNTIME_API virtual void BeginDestroy() override;
+	TYPEDELEMENTRUNTIME_API virtual void Serialize(FArchive& Ar) override;
 
 	/**
 	 * Test to see whether the given element is currently considered selected.
 	 */
-	bool IsElementSelected(const FTypedElementHandle& InElementHandle, const FTypedElementIsSelectedOptions InSelectionOptions) const;
+	TYPEDELEMENTRUNTIME_API bool IsElementSelected(const FTypedElementHandle& InElementHandle, const FTypedElementIsSelectedOptions InSelectionOptions) const;
 
 	/**
 	 * Test to see whether the given element can be selected.
 	 */
-	bool CanSelectElement(const FTypedElementHandle& InElementHandle, const FTypedElementSelectionOptions InSelectionOptions) const;
+	TYPEDELEMENTRUNTIME_API bool CanSelectElement(const FTypedElementHandle& InElementHandle, const FTypedElementSelectionOptions InSelectionOptions) const;
 
 	/**
 	 * Test to see whether the given element can be deselected.
 	 */
-	bool CanDeselectElement(const FTypedElementHandle& InElementHandle, const FTypedElementSelectionOptions InSelectionOptions) const;
+	TYPEDELEMENTRUNTIME_API bool CanDeselectElement(const FTypedElementHandle& InElementHandle, const FTypedElementSelectionOptions InSelectionOptions) const;
 
 	/**
 	 * Attempt to select the given element.
 	 * @return True if the selection was changed, false otherwise.
 	 */
-	bool SelectElement(const FTypedElementHandle& InElementHandle, const FTypedElementSelectionOptions InSelectionOptions);
+	TYPEDELEMENTRUNTIME_API bool SelectElement(const FTypedElementHandle& InElementHandle, const FTypedElementSelectionOptions InSelectionOptions);
 
 	/**
 	 * Attempt to select the given elements.
 	 * @return True if the selection was changed, false otherwise.
 	 */
-	bool SelectElements(const TArray<FTypedElementHandle>& InElementHandles, const FTypedElementSelectionOptions InSelectionOptions);
-	bool SelectElements(TArrayView<const FTypedElementHandle> InElementHandles, const FTypedElementSelectionOptions InSelectionOptions);
-	bool SelectElements(FTypedElementListConstRef InElementList, const FTypedElementSelectionOptions InSelectionOptions);
+	TYPEDELEMENTRUNTIME_API bool SelectElements(const TArray<FTypedElementHandle>& InElementHandles, const FTypedElementSelectionOptions InSelectionOptions);
+	TYPEDELEMENTRUNTIME_API bool SelectElements(TArrayView<const FTypedElementHandle> InElementHandles, const FTypedElementSelectionOptions InSelectionOptions);
+	TYPEDELEMENTRUNTIME_API bool SelectElements(FTypedElementListConstRef InElementList, const FTypedElementSelectionOptions InSelectionOptions);
 
 	/**
 	 * Attempt to deselect the given element.
 	 * @return True if the selection was changed, false otherwise.
 	 */
-	bool DeselectElement(const FTypedElementHandle& InElementHandle, const FTypedElementSelectionOptions InSelectionOptions);
+	TYPEDELEMENTRUNTIME_API bool DeselectElement(const FTypedElementHandle& InElementHandle, const FTypedElementSelectionOptions InSelectionOptions);
 
 	/**
 	 * Attempt to deselect the given elements.
 	 * @return True if the selection was changed, false otherwise.
 	 */
-	bool DeselectElements(const TArray<FTypedElementHandle>& InElementHandles, const FTypedElementSelectionOptions InSelectionOptions);
-	bool DeselectElements(TArrayView<const FTypedElementHandle> InElementHandles, const FTypedElementSelectionOptions InSelectionOptions);
-	bool DeselectElements(FTypedElementListConstRef InElementList, const FTypedElementSelectionOptions InSelectionOptions);
+	TYPEDELEMENTRUNTIME_API bool DeselectElements(const TArray<FTypedElementHandle>& InElementHandles, const FTypedElementSelectionOptions InSelectionOptions);
+	TYPEDELEMENTRUNTIME_API bool DeselectElements(TArrayView<const FTypedElementHandle> InElementHandles, const FTypedElementSelectionOptions InSelectionOptions);
+	TYPEDELEMENTRUNTIME_API bool DeselectElements(FTypedElementListConstRef InElementList, const FTypedElementSelectionOptions InSelectionOptions);
 
 	/**
 	 * Clear the current selection.
 	 * @return True if the selection was changed, false otherwise.
 	 */
 	UFUNCTION(BlueprintCallable, Category="TypedElementFramework|Selection")
-	bool ClearSelection(const FTypedElementSelectionOptions InSelectionOptions);
+	TYPEDELEMENTRUNTIME_API bool ClearSelection(const FTypedElementSelectionOptions InSelectionOptions);
 
 	/**
 	 * Attempt to make the selection the given elements.
 	 * @note Equivalent to calling ClearSelection then SelectElements, but happens in a single batch.
 	 * @return True if the selection was changed, false otherwise.
 	 */
-	bool SetSelection(const TArray<FTypedElementHandle>& InElementHandles, const FTypedElementSelectionOptions InSelectionOptions);
-	bool SetSelection(TArrayView<const FTypedElementHandle> InElementHandles, const FTypedElementSelectionOptions InSelectionOptions);
-	bool SetSelection(FTypedElementListConstRef InElementList, const FTypedElementSelectionOptions InSelectionOptions);
+	TYPEDELEMENTRUNTIME_API bool SetSelection(const TArray<FTypedElementHandle>& InElementHandles, const FTypedElementSelectionOptions InSelectionOptions);
+	TYPEDELEMENTRUNTIME_API bool SetSelection(TArrayView<const FTypedElementHandle> InElementHandles, const FTypedElementSelectionOptions InSelectionOptions);
+	TYPEDELEMENTRUNTIME_API bool SetSelection(FTypedElementListConstRef InElementList, const FTypedElementSelectionOptions InSelectionOptions);
 
 	/**
 	 * Test to see whether selection modifiers (Ctrl or Shift) are allowed while selecting this element.
 	 */
-	bool AllowSelectionModifiers(const FTypedElementHandle& InElementHandle) const;
+	TYPEDELEMENTRUNTIME_API bool AllowSelectionModifiers(const FTypedElementHandle& InElementHandle) const;
 
 	/**
 	 * Given an element, return the element that should actually perform a selection operation.
 	 */
-	FTypedElementHandle GetSelectionElement(const FTypedElementHandle& InElementHandle, const ETypedElementSelectionMethod InSelectionMethod) const;
+	TYPEDELEMENTRUNTIME_API FTypedElementHandle GetSelectionElement(const FTypedElementHandle& InElementHandle, const ETypedElementSelectionMethod InSelectionMethod) const;
 
 	/**
 	 * Get a normalized version of this selection set that can be used to perform operations like gizmo manipulation, deletion, copying, etc.
 	 * This will do things like expand out groups, and resolve any parent<->child elements so that duplication operations aren't performed on both the parent and the child.
 	 */
-	FTypedElementListRef GetNormalizedSelection(const FTypedElementSelectionNormalizationOptions InNormalizationOptions) const;
-	void GetNormalizedSelection(const FTypedElementSelectionNormalizationOptions& InNormalizationOptions, FTypedElementListRef OutNormalizedElements) const;
+	TYPEDELEMENTRUNTIME_API FTypedElementListRef GetNormalizedSelection(const FTypedElementSelectionNormalizationOptions InNormalizationOptions) const;
+	TYPEDELEMENTRUNTIME_API void GetNormalizedSelection(const FTypedElementSelectionNormalizationOptions& InNormalizationOptions, FTypedElementListRef OutNormalizedElements) const;
 
 	/**
 	 * Get a normalized version of the given element list that can be used to perform operations like gizmo manipulation, deletion, copying, etc.
 	 * This will do things like expand out groups, and resolve any parent<->child elements so that duplication operations aren't performed on both the parent and the child.
 	 */
-	FTypedElementListRef GetNormalizedElementList(FTypedElementListConstRef InElementList, const FTypedElementSelectionNormalizationOptions InNormalizationOptions) const;
-	void GetNormalizedElementList(FTypedElementListConstRef InElementList, const FTypedElementSelectionNormalizationOptions& InNormalizationOptions, FTypedElementListRef OutNormalizedElements) const;
+	TYPEDELEMENTRUNTIME_API FTypedElementListRef GetNormalizedElementList(FTypedElementListConstRef InElementList, const FTypedElementSelectionNormalizationOptions InNormalizationOptions) const;
+	TYPEDELEMENTRUNTIME_API void GetNormalizedElementList(FTypedElementListConstRef InElementList, const FTypedElementSelectionNormalizationOptions& InNormalizationOptions, FTypedElementListRef OutNormalizedElements) const;
 
 	/**
 	 * Get the number of selected elements.
@@ -548,14 +548,14 @@ public:
 	 * @returns the current state of the selection set.
 	 */
 	UFUNCTION(BlueprintPure, Category = "TypedElementFramework|Selection")
-	FTypedElementSelectionSetState GetCurrentSelectionState() const;
+	TYPEDELEMENTRUNTIME_API FTypedElementSelectionSetState GetCurrentSelectionState() const;
 
 	/**
 	 * Restores the selection set from the given state.
 	 * The calling code is responsible for managing any undo state.
 	 */
 	UFUNCTION(BlueprintCallable, Category = "TypedElementFramework|Selection")
-	void RestoreSelectionState(const FTypedElementSelectionSetState& InSelectionState);
+	TYPEDELEMENTRUNTIME_API void RestoreSelectionState(const FTypedElementSelectionSetState& InSelectionState);
 
 
 	/*
@@ -566,47 +566,47 @@ public:
 	 * Test to see whether the given element is currently considered selected.
 	 */
 	UFUNCTION(BlueprintPure, Category="TypedElementFramework|Selection")
-	bool IsElementSelected(const FScriptTypedElementHandle& InElementHandle, const FTypedElementIsSelectedOptions InSelectionOptions) const;
+	TYPEDELEMENTRUNTIME_API bool IsElementSelected(const FScriptTypedElementHandle& InElementHandle, const FTypedElementIsSelectedOptions InSelectionOptions) const;
 
 	/**
 	 * Test to see whether the given element can be selected.
 	 */
 	UFUNCTION(BlueprintPure, Category="TypedElementFramework|Selection")
-	bool CanSelectElement(const FScriptTypedElementHandle& InElementHandle, const FTypedElementSelectionOptions InSelectionOptions) const;
+	TYPEDELEMENTRUNTIME_API bool CanSelectElement(const FScriptTypedElementHandle& InElementHandle, const FTypedElementSelectionOptions InSelectionOptions) const;
 
 	/**
 	 * Test to see whether the given element can be deselected.
 	 */
 	UFUNCTION(BlueprintPure, Category="TypedElementFramework|Selection")
-	bool CanDeselectElement(const FScriptTypedElementHandle& InElementHandle, const FTypedElementSelectionOptions InSelectionOptions) const;
+	TYPEDELEMENTRUNTIME_API bool CanDeselectElement(const FScriptTypedElementHandle& InElementHandle, const FTypedElementSelectionOptions InSelectionOptions) const;
 
 	/**
 	 * Attempt to select the given element.
 	 * @return True if the selection was changed, false otherwise.
 	 */
 	UFUNCTION(BlueprintCallable, Category="TypedElementFramework|Selection")
-	bool SelectElement(const FScriptTypedElementHandle& InElementHandle, const FTypedElementSelectionOptions InSelectionOptions);
+	TYPEDELEMENTRUNTIME_API bool SelectElement(const FScriptTypedElementHandle& InElementHandle, const FTypedElementSelectionOptions InSelectionOptions);
 
 	/**
 	 * Attempt to select the given elements.
 	 * @return True if the selection was changed, false otherwise.
 	 */
 	UFUNCTION(BlueprintCallable, Category="TypedElementFramework|Selection")
-	bool SelectElements(const TArray<FScriptTypedElementHandle>& InElementHandles, const FTypedElementSelectionOptions InSelectionOptions);
+	TYPEDELEMENTRUNTIME_API bool SelectElements(const TArray<FScriptTypedElementHandle>& InElementHandles, const FTypedElementSelectionOptions InSelectionOptions);
 
 	/**
 	 * Attempt to deselect the given element.
 	 * @return True if the selection was changed, false otherwise.
 	 */
 	UFUNCTION(BlueprintCallable, Category="TypedElementFramework|Selection")
-	bool DeselectElement(const FScriptTypedElementHandle& InElementHandle, const FTypedElementSelectionOptions InSelectionOptions);
+	TYPEDELEMENTRUNTIME_API bool DeselectElement(const FScriptTypedElementHandle& InElementHandle, const FTypedElementSelectionOptions InSelectionOptions);
 
 	/**
 	 * Attempt to deselect the given elements.
 	 * @return True if the selection was changed, false otherwise.
 	 */
 	UFUNCTION(BlueprintCallable, Category="TypedElementFramework|Selection")
-	bool DeselectElements(const TArray<FScriptTypedElementHandle>& InElementHandles, const FTypedElementSelectionOptions InSelectionOptions);
+	TYPEDELEMENTRUNTIME_API bool DeselectElements(const TArray<FScriptTypedElementHandle>& InElementHandles, const FTypedElementSelectionOptions InSelectionOptions);
 
 	/**
 	 * Attempt to make the selection the given elements.
@@ -614,26 +614,26 @@ public:
 	 * @return True if the selection was changed, false otherwise.
 	 */
 	UFUNCTION(BlueprintCallable, Category="TypedElementFramework|Selection")
-	bool SetSelection(const TArray<FScriptTypedElementHandle>& InElementHandles, const FTypedElementSelectionOptions InSelectionOptions);
+	TYPEDELEMENTRUNTIME_API bool SetSelection(const TArray<FScriptTypedElementHandle>& InElementHandles, const FTypedElementSelectionOptions InSelectionOptions);
 
 	/**
 	 * Test to see whether selection modifiers (Ctrl or Shift) are allowed while selecting this element.
 	 */
 	UFUNCTION(BlueprintPure, Category="TypedElementFramework|Selection")
-	bool AllowSelectionModifiers(const FScriptTypedElementHandle& InElementHandle) const;
+	TYPEDELEMENTRUNTIME_API bool AllowSelectionModifiers(const FScriptTypedElementHandle& InElementHandle) const;
 
 	/**
 	 * Given an element, return the element that should actually perform a selection operation.
 	 */
 	UFUNCTION(BlueprintPure, Category = "TypedElementFramework|Selection")
-	FScriptTypedElementHandle GetSelectionElement(const FScriptTypedElementHandle& InElementHandle, const ETypedElementSelectionMethod InSelectionMethod) const;
+	TYPEDELEMENTRUNTIME_API FScriptTypedElementHandle GetSelectionElement(const FScriptTypedElementHandle& InElementHandle, const ETypedElementSelectionMethod InSelectionMethod) const;
 
 
 	/**
 	 * Get the handle of every selected element, optionally filtering to elements that implement the given interface.
 	 */
 	UFUNCTION(BlueprintCallable, BlueprintPure=false, DisplayName="Get Selected Element Handles", Category="TypedElementFramework|Selection", meta=(ScriptName="GetSelectedElementHandles"))
-	TArray<FScriptTypedElementHandle> K2_GetSelectedElementHandles(const TSubclassOf<UInterface> InBaseInterfaceType = nullptr) const;
+	TYPEDELEMENTRUNTIME_API TArray<FScriptTypedElementHandle> K2_GetSelectedElementHandles(const TSubclassOf<UInterface> InBaseInterfaceType = nullptr) const;
 
 	DECLARE_DYNAMIC_MULTICAST_DELEGATE_OneParam(FOnPreChangeDynamic, const UTypedElementSelectionSet*, SelectionSet);
 
@@ -651,27 +651,27 @@ private:
 	/**
 	 * Attempt to resolve the selection interface and selection customization for the given element, if any.
 	 */
-	FTypedElementSelectionSetElement ResolveSelectionSetElement(const FTypedElementHandle& InElementHandle) const;
+	TYPEDELEMENTRUNTIME_API FTypedElementSelectionSetElement ResolveSelectionSetElement(const FTypedElementHandle& InElementHandle) const;
 
 	/**
 	 * Update the selection from a replacement request.
 	 */
-	void OnElementReplaced(TArrayView<const TTuple<FTypedElementHandle, FTypedElementHandle>> InReplacedElements);
+	TYPEDELEMENTRUNTIME_API void OnElementReplaced(TArrayView<const TTuple<FTypedElementHandle, FTypedElementHandle>> InReplacedElements);
 
 	/**
 	 * Force a selection update if an element internal state changes.
 	 */
-	void OnElementUpdated(TArrayView<const FTypedElementHandle> InUpdatedElements);
+	TYPEDELEMENTRUNTIME_API void OnElementUpdated(TArrayView<const FTypedElementHandle> InUpdatedElements);
 
 	/**
 	 * Proxy the internal OnPreChange event from the underlying element list.
 	 */
-	void OnElementListPreChange(const FTypedElementList& InElementList);
+	TYPEDELEMENTRUNTIME_API void OnElementListPreChange(const FTypedElementList& InElementList);
 
 	/**
 	 * Proxy the internal OnChanged event from the underlying element list.
 	 */
-	void OnElementListChanged(const FTypedElementList& InElementList);
+	TYPEDELEMENTRUNTIME_API void OnElementListChanged(const FTypedElementList& InElementList);
 
 	/** Underlying element list holding the selection state. */
 	FTypedElementListPtr ElementList;
diff --git a/Engine/Source/Runtime/TypedElementRuntime/Public/Elements/Interfaces/TypedElementAssetDataInterface.h b/Engine/Source/Runtime/TypedElementRuntime/Public/Elements/Interfaces/TypedElementAssetDataInterface.h
--- a/Engine/Source/Runtime/TypedElementRuntime/Public/Elements/Interfaces/TypedElementAssetDataInterface.h
+++ b/Engine/Source/Runtime/TypedElementRuntime/Public/Elements/Interfaces/TypedElementAssetDataInterface.h
@@ -20,7 +20,7 @@ class UTypedElementAssetDataInterface : public UInterface
 	GENERATED_BODY()
 };
 
-class TYPEDELEMENTRUNTIME_API ITypedElementAssetDataInterface
+class ITypedElementAssetDataInterface
 {
 	GENERATED_BODY()
 
@@ -31,12 +31,12 @@ public:
 	 *
 	 * @returns An array of valid asset datas.
 	 */
-	virtual TArray<FAssetData> GetAllReferencedAssetDatas(const FTypedElementHandle& InElementHandle);
+	TYPEDELEMENTRUNTIME_API virtual TArray<FAssetData> GetAllReferencedAssetDatas(const FTypedElementHandle& InElementHandle);
 
 	/**
 	 * Returns the asset data for the given handle, if it exists.
 	 */
-	virtual FAssetData GetAssetData(const FTypedElementHandle& InElementHandle);
+	TYPEDELEMENTRUNTIME_API virtual FAssetData GetAssetData(const FTypedElementHandle& InElementHandle);
 
 	/**
 	 * Script Api
@@ -49,13 +49,13 @@ public:
 	 * @returns An array of valid asset datas.
 	 */
 	UFUNCTION(BlueprintCallable, Category = "TypedElementInterfaces|AssetData")
-	virtual TArray<FAssetData> GetAllReferencedAssetDatas(const FScriptTypedElementHandle& InElementHandle);
+	TYPEDELEMENTRUNTIME_API virtual TArray<FAssetData> GetAllReferencedAssetDatas(const FScriptTypedElementHandle& InElementHandle);
 
 	/**
 	 * Returns the asset data for the given handle, if it exists.
 	 */
 	UFUNCTION(BlueprintCallable, Category = "TypedElementInterfaces|AssetData")
-	virtual FAssetData GetAssetData(const FScriptTypedElementHandle& InElementHandle);
+	TYPEDELEMENTRUNTIME_API virtual FAssetData GetAssetData(const FScriptTypedElementHandle& InElementHandle);
 };
 
 template <>
diff --git a/Engine/Source/Runtime/TypedElementRuntime/Public/Elements/Interfaces/TypedElementHierarchyInterface.h b/Engine/Source/Runtime/TypedElementRuntime/Public/Elements/Interfaces/TypedElementHierarchyInterface.h
--- a/Engine/Source/Runtime/TypedElementRuntime/Public/Elements/Interfaces/TypedElementHierarchyInterface.h
+++ b/Engine/Source/Runtime/TypedElementRuntime/Public/Elements/Interfaces/TypedElementHierarchyInterface.h
@@ -19,7 +19,7 @@ class UTypedElementHierarchyInterface : public UInterface
 	GENERATED_BODY()
 };
 
-class TYPEDELEMENTRUNTIME_API ITypedElementHierarchyInterface
+class ITypedElementHierarchyInterface
 {
 	GENERATED_BODY()
 
@@ -47,7 +47,7 @@ public:
 	 * eg) A component might return its actor, or a static mesh instance might return its ISM component.
 	 */
 	UFUNCTION(BlueprintCallable, Category="TypedElementInterfaces|Hierarchy")
-	virtual FScriptTypedElementHandle GetParentElement(const FScriptTypedElementHandle& InElementHandle, const bool bAllowCreate = true);
+	TYPEDELEMENTRUNTIME_API virtual FScriptTypedElementHandle GetParentElement(const FScriptTypedElementHandle& InElementHandle, const bool bAllowCreate = true);
 
 	/**
 	 * Get the logical children of this element, if any.
@@ -56,14 +56,14 @@ public:
 	 * @note Appends to OutElementHandles.
 	 */
 	UFUNCTION(BlueprintCallable, Category="TypedElementInterfaces|Hierarchy")
-	virtual void GetChildElements(const FScriptTypedElementHandle& InElementHandle, TArray<FScriptTypedElementHandle>& OutElementHandles, const bool bAllowCreate = true);
+	TYPEDELEMENTRUNTIME_API virtual void GetChildElements(const FScriptTypedElementHandle& InElementHandle, TArray<FScriptTypedElementHandle>& OutElementHandles, const bool bAllowCreate = true);
 
 private:
 	
 	/**
 	 * Return the registry used when creating new elements.
 	 */
-	virtual class UTypedElementRegistry& GetRegistry() const;
+	TYPEDELEMENTRUNTIME_API virtual class UTypedElementRegistry& GetRegistry() const;
 };
 
 template <>
diff --git a/Engine/Source/Runtime/TypedElementRuntime/Public/Elements/Interfaces/TypedElementObjectInterface.h b/Engine/Source/Runtime/TypedElementRuntime/Public/Elements/Interfaces/TypedElementObjectInterface.h
--- a/Engine/Source/Runtime/TypedElementRuntime/Public/Elements/Interfaces/TypedElementObjectInterface.h
+++ b/Engine/Source/Runtime/TypedElementRuntime/Public/Elements/Interfaces/TypedElementObjectInterface.h
@@ -21,7 +21,7 @@ class UTypedElementObjectInterface : public UInterface
 	GENERATED_BODY()
 };
 
-class TYPEDELEMENTRUNTIME_API ITypedElementObjectInterface
+class ITypedElementObjectInterface
 {
 	GENERATED_BODY()
 
@@ -29,12 +29,12 @@ public:
 	/**
 	 * Get the object instance that this handle represents, if any.
 	 */
-	virtual UObject* GetObject(const FTypedElementHandle& InElementHandle);
+	TYPEDELEMENTRUNTIME_API virtual UObject* GetObject(const FTypedElementHandle& InElementHandle);
 
 	/**
 	 * Gets the object instance's class that the handle represents, if any. 
 	 */
-	virtual UClass* GetObjectClass(const FTypedElementHandle& InElementHandle);
+	TYPEDELEMENTRUNTIME_API virtual UClass* GetObjectClass(const FTypedElementHandle& InElementHandle);
 
 	/**
 	 * Attempts to cast the given handle to another class.
@@ -79,13 +79,13 @@ public:
 	 * Get the object instance that this handle represents, if any.
 	 */
 	UFUNCTION(BlueprintCallable, Category="TypedElementInterfaces|Object")
-	virtual UObject* GetObject(const FScriptTypedElementHandle& InElementHandle);
+	TYPEDELEMENTRUNTIME_API virtual UObject* GetObject(const FScriptTypedElementHandle& InElementHandle);
 
 	/**
 	 * Gets the object instance's class that the handle represents, if any. 
 	 */
 	UFUNCTION(BlueprintCallable, Category = "TypedElementInterfaces|Object")
-	virtual UClass* GetObjectClass(const FScriptTypedElementHandle& InElementHandle);
+	TYPEDELEMENTRUNTIME_API virtual UClass* GetObjectClass(const FScriptTypedElementHandle& InElementHandle);
 };
 
 template <>
diff --git a/Engine/Source/Runtime/TypedElementRuntime/Public/Elements/Interfaces/TypedElementPrimitiveCustomDataInterface.h b/Engine/Source/Runtime/TypedElementRuntime/Public/Elements/Interfaces/TypedElementPrimitiveCustomDataInterface.h
--- a/Engine/Source/Runtime/TypedElementRuntime/Public/Elements/Interfaces/TypedElementPrimitiveCustomDataInterface.h
+++ b/Engine/Source/Runtime/TypedElementRuntime/Public/Elements/Interfaces/TypedElementPrimitiveCustomDataInterface.h
@@ -1,4 +1,4 @@
-﻿// Copyright Epic Games, Inc. All Rights Reserved.
+// Copyright Epic Games, Inc. All Rights Reserved.
 
 #pragma once
 
@@ -17,7 +17,7 @@ class UTypedElementPrimitiveCustomDataInterface : public UInterface
  * An interface for use with the TypedElement Framework which exposes Primitive CustomData
  * for use within Materials.
  */
-class TYPEDELEMENTRUNTIME_API ITypedElementPrimitiveCustomDataInterface
+class ITypedElementPrimitiveCustomDataInterface
 {
 	GENERATED_BODY()
 
@@ -30,11 +30,11 @@ public:
 
 	// Sets all Primitive's CustomData values
 	UFUNCTION(BlueprintCallable, Category="TypedElementInterfaces|CustomData")
-	virtual void SetCustomData(const FScriptTypedElementHandle& InElementHandle, const TArray<float>& CustomDataFloats,  bool bMarkRenderStateDirty = false);
+	TYPEDELEMENTRUNTIME_API virtual void SetCustomData(const FScriptTypedElementHandle& InElementHandle, const TArray<float>& CustomDataFloats,  bool bMarkRenderStateDirty = false);
 
 	// Sets a single Primitive's CustomData value
 	UFUNCTION(BlueprintCallable, Category="TypedElementInterfaces|CustomData")
-	virtual void SetCustomDataValue(const FScriptTypedElementHandle& InElementHandle, int32 CustomDataIndex, float CustomDataValue, bool bMarkRenderStateDirty = false);
+	TYPEDELEMENTRUNTIME_API virtual void SetCustomDataValue(const FScriptTypedElementHandle& InElementHandle, int32 CustomDataIndex, float CustomDataValue, bool bMarkRenderStateDirty = false);
 };
 
 template <>
@@ -42,4 +42,4 @@ struct TTypedElement<ITypedElementPrimitiveCustomDataInterface> : public TTypedE
 {
 	void SetCustomData(TArrayView<float> CustomDataFloats,  bool bMarkRenderStateDirty = false) const { return InterfacePtr->SetCustomData(*this, CustomDataFloats, bMarkRenderStateDirty); }
 	void SetCustomDataValue(int32 CustomDataIndex, float CustomDataValue, bool bMarkRenderStateDirty = false) { return InterfacePtr->SetCustomDataValue(*this, CustomDataIndex, CustomDataValue, bMarkRenderStateDirty); }
-};
\ No newline at end of file
+};
diff --git a/Engine/Source/Runtime/TypedElementRuntime/Public/Elements/Interfaces/TypedElementSelectionInterface.h b/Engine/Source/Runtime/TypedElementRuntime/Public/Elements/Interfaces/TypedElementSelectionInterface.h
--- a/Engine/Source/Runtime/TypedElementRuntime/Public/Elements/Interfaces/TypedElementSelectionInterface.h
+++ b/Engine/Source/Runtime/TypedElementRuntime/Public/Elements/Interfaces/TypedElementSelectionInterface.h
@@ -147,7 +147,7 @@ class UTypedElementSelectionInterface : public UInterface
 	GENERATED_BODY()
 };
 
-class TYPEDELEMENTRUNTIME_API ITypedElementSelectionInterface
+class ITypedElementSelectionInterface
 {
 	GENERATED_BODY()
 
@@ -155,7 +155,7 @@ public:
 	/**
 	 * Test to see whether the given element is currently considered selected.
 	 */
-	virtual bool IsElementSelected(const FTypedElementHandle& InElementHandle, const FTypedElementListConstPtr& SelectionSetPtr, const FTypedElementIsSelectedOptions& InSelectionOptions);
+	TYPEDELEMENTRUNTIME_API virtual bool IsElementSelected(const FTypedElementHandle& InElementHandle, const FTypedElementListConstPtr& SelectionSetPtr, const FTypedElementIsSelectedOptions& InSelectionOptions);
 
 	/**
 	 * Test to see whether the given element can be selected.
@@ -171,13 +171,13 @@ public:
 	 * Attempt to select the given element.
 	 * @return True if the selection was changed, false otherwise.
 	 */
-	virtual bool SelectElement(const FTypedElementHandle& InElementHandle, const FTypedElementListPtr& InSelectionSet, const FTypedElementSelectionOptions& InSelectionOptions);
+	TYPEDELEMENTRUNTIME_API virtual bool SelectElement(const FTypedElementHandle& InElementHandle, const FTypedElementListPtr& InSelectionSet, const FTypedElementSelectionOptions& InSelectionOptions);
 
 	/**
 	 * Attempt to deselect the given element.
 	 * @return True if the selection was changed, false otherwise.
 	 */
-	virtual bool DeselectElement(const FTypedElementHandle& InElementHandle, const FTypedElementListPtr& InSelectionSet, const FTypedElementSelectionOptions& InSelectionOptions);
+	TYPEDELEMENTRUNTIME_API virtual bool DeselectElement(const FTypedElementHandle& InElementHandle, const FTypedElementListPtr& InSelectionSet, const FTypedElementSelectionOptions& InSelectionOptions);
 
 	/**
 	 * Test to see whether selection modifiers (Ctrl or Shift) are allowed while selecting this element.
@@ -229,45 +229,45 @@ public:
 	 * Test to see whether the given element is currently considered selected.
 	 */
 	UFUNCTION(BlueprintCallable, Category = "TypedElementInterfaces|Selection")
-	virtual bool IsElementSelected(const FScriptTypedElementHandle& InElementHandle, const FScriptTypedElementListProxy InSelectionSet, const FTypedElementIsSelectedOptions& InSelectionOptions);
+	TYPEDELEMENTRUNTIME_API virtual bool IsElementSelected(const FScriptTypedElementHandle& InElementHandle, const FScriptTypedElementListProxy InSelectionSet, const FTypedElementIsSelectedOptions& InSelectionOptions);
 
 	/**
 	 * Test to see whether the given element can be selected.
 	 */
 	UFUNCTION(BlueprintCallable, Category="TypedElementInterfaces|Selection")
-	virtual bool CanSelectElement(const FScriptTypedElementHandle& InElementHandle, const FTypedElementSelectionOptions& InSelectionOptions);
+	TYPEDELEMENTRUNTIME_API virtual bool CanSelectElement(const FScriptTypedElementHandle& InElementHandle, const FTypedElementSelectionOptions& InSelectionOptions);
 
 	/**
 	 * Test to see whether the given element can be deselected.
 	 */
 	UFUNCTION(BlueprintCallable, Category="TypedElementInterfaces|Selection")
-	virtual bool CanDeselectElement(const FScriptTypedElementHandle& InElementHandle, const FTypedElementSelectionOptions& InSelectionOptions);
+	TYPEDELEMENTRUNTIME_API virtual bool CanDeselectElement(const FScriptTypedElementHandle& InElementHandle, const FTypedElementSelectionOptions& InSelectionOptions);
 
 	/**
 	 * Attempt to select the given element.
 	 * @return True if the selection was changed, false otherwise.
 	 */
 	UFUNCTION(BlueprintCallable, Category="TypedElementInterfaces|Selection")
-	virtual bool SelectElement(const FScriptTypedElementHandle& InElementHandle, FScriptTypedElementListProxy InSelectionSet, const FTypedElementSelectionOptions& InSelectionOptions);
+	TYPEDELEMENTRUNTIME_API virtual bool SelectElement(const FScriptTypedElementHandle& InElementHandle, FScriptTypedElementListProxy InSelectionSet, const FTypedElementSelectionOptions& InSelectionOptions);
 
 	/**
 	 * Attempt to deselect the given element.
 	 * @return True if the selection was changed, false otherwise.
 	 */
 	UFUNCTION(BlueprintCallable, Category="TypedElementInterfaces|Selection")
-	virtual bool DeselectElement(const FScriptTypedElementHandle& InElementHandle, FScriptTypedElementListProxy InSelectionSet, const FTypedElementSelectionOptions& InSelectionOptions);
+	TYPEDELEMENTRUNTIME_API virtual bool DeselectElement(const FScriptTypedElementHandle& InElementHandle, FScriptTypedElementListProxy InSelectionSet, const FTypedElementSelectionOptions& InSelectionOptions);
 
 	/**
 	 * Test to see whether selection modifiers (Ctrl or Shift) are allowed while selecting this element.
 	 */
 	UFUNCTION(BlueprintCallable, Category="TypedElementInterfaces|Selection")
-	virtual bool AllowSelectionModifiers(const FScriptTypedElementHandle& InElementHandle, const FScriptTypedElementListProxy InSelectionSet);
+	TYPEDELEMENTRUNTIME_API virtual bool AllowSelectionModifiers(const FScriptTypedElementHandle& InElementHandle, const FScriptTypedElementListProxy InSelectionSet);
 
 	/**
 	 * Given an element, return the element that should actually perform a selection operation.
 	 */
 	UFUNCTION(BlueprintCallable, Category="TypedElementInterfaces|Selection")
-	virtual FScriptTypedElementHandle GetSelectionElement(const FScriptTypedElementHandle& InElementHandle, const FScriptTypedElementListProxy InCurrentSelection, const ETypedElementSelectionMethod InSelectionMethod);
+	TYPEDELEMENTRUNTIME_API virtual FScriptTypedElementHandle GetSelectionElement(const FScriptTypedElementHandle& InElementHandle, const FScriptTypedElementListProxy InCurrentSelection, const ETypedElementSelectionMethod InSelectionMethod);
 
 protected:
 	/**
diff --git a/Engine/Source/Runtime/UMG/Private/Animation/MovieSceneMarginPropertySystem.h b/Engine/Source/Runtime/UMG/Private/Animation/MovieSceneMarginPropertySystem.h
--- a/Engine/Source/Runtime/UMG/Private/Animation/MovieSceneMarginPropertySystem.h
+++ b/Engine/Source/Runtime/UMG/Private/Animation/MovieSceneMarginPropertySystem.h
@@ -8,15 +8,15 @@
 
 Expose_TNameOf(FMargin);
 
-UCLASS()
-class UMG_API UMovieSceneMarginPropertySystem : public UMovieScenePropertySystem
+UCLASS(MinimalAPI)
+class UMovieSceneMarginPropertySystem : public UMovieScenePropertySystem
 {
 public:
 
 	GENERATED_BODY()
 	
-	UMovieSceneMarginPropertySystem(const FObjectInitializer& ObjInit);
+	UMG_API UMovieSceneMarginPropertySystem(const FObjectInitializer& ObjInit);
 
-	virtual void OnRun(FSystemTaskPrerequisites& InPrerequisites, FSystemSubsequentTasks& Subsequents) override;
+	UMG_API virtual void OnRun(FSystemTaskPrerequisites& InPrerequisites, FSystemSubsequentTasks& Subsequents) override;
 };
 
diff --git a/Engine/Source/Runtime/UMG/Public/Animation/MovieScene2DTransformPropertySystem.h b/Engine/Source/Runtime/UMG/Public/Animation/MovieScene2DTransformPropertySystem.h
--- a/Engine/Source/Runtime/UMG/Public/Animation/MovieScene2DTransformPropertySystem.h
+++ b/Engine/Source/Runtime/UMG/Public/Animation/MovieScene2DTransformPropertySystem.h
@@ -5,14 +5,14 @@
 #include "Systems/MovieScenePropertySystem.h"
 #include "MovieScene2DTransformPropertySystem.generated.h"
 
-UCLASS()
-class UMG_API UMovieScene2DTransformPropertySystem : public UMovieScenePropertySystem
+UCLASS(MinimalAPI)
+class UMovieScene2DTransformPropertySystem : public UMovieScenePropertySystem
 {
 public:
 
 	GENERATED_BODY()
 
-	UMovieScene2DTransformPropertySystem(const FObjectInitializer& ObjInit);
+	UMG_API UMovieScene2DTransformPropertySystem(const FObjectInitializer& ObjInit);
 
-	virtual void OnRun(FSystemTaskPrerequisites& InPrerequisites, FSystemSubsequentTasks& Subsequents) override;
+	UMG_API virtual void OnRun(FSystemTaskPrerequisites& InPrerequisites, FSystemSubsequentTasks& Subsequents) override;
 };
diff --git a/Engine/Source/Runtime/UMG/Public/Animation/MovieSceneUMGComponentTypes.h b/Engine/Source/Runtime/UMG/Public/Animation/MovieSceneUMGComponentTypes.h
--- a/Engine/Source/Runtime/UMG/Public/Animation/MovieSceneUMGComponentTypes.h
+++ b/Engine/Source/Runtime/UMG/Public/Animation/MovieSceneUMGComponentTypes.h
@@ -54,9 +54,9 @@ UMG_API void ConvertOperationalProperty(const FMargin& In, FIntermediateMargin&
 using FMarginTraits = TIndirectPropertyTraits<FMargin, FIntermediateMargin>;
 using FWidgetTransformPropertyTraits = TIndirectPropertyTraits<FWidgetTransform, FIntermediateWidgetTransform>;
 
-struct UMG_API FMovieSceneUMGComponentTypes
+struct FMovieSceneUMGComponentTypes
 {
-	~FMovieSceneUMGComponentTypes();
+	UMG_API ~FMovieSceneUMGComponentTypes();
 
 	TPropertyComponents<FMarginTraits> Margin;
 	TPropertyComponents<FWidgetTransformPropertyTraits> WidgetTransform;
@@ -67,9 +67,9 @@ struct UMG_API FMovieSceneUMGComponentTypes
 	TCustomPropertyRegistration<FMarginTraits, 1> CustomMarginAccessors;
 	TCustomPropertyRegistration<FWidgetTransformPropertyTraits, 1> CustomWidgetTransformAccessors;
 
-	static void Destroy();
+	static UMG_API void Destroy();
 
-	static FMovieSceneUMGComponentTypes* Get();
+	static UMG_API FMovieSceneUMGComponentTypes* Get();
 
 private:
 	FMovieSceneUMGComponentTypes();
diff --git a/Engine/Source/Runtime/UMG/Public/Animation/UMGSequencePlayer.h b/Engine/Source/Runtime/UMG/Public/Animation/UMGSequencePlayer.h
--- a/Engine/Source/Runtime/UMG/Public/Animation/UMGSequencePlayer.h
+++ b/Engine/Source/Runtime/UMG/Public/Animation/UMGSequencePlayer.h
@@ -14,31 +14,31 @@
 
 class UWidgetAnimation;
 
-UCLASS(Transient, BlueprintType)
-class UMG_API UUMGSequencePlayer : public UObject, public IMovieScenePlayer
+UCLASS(Transient, BlueprintType, MinimalAPI)
+class UUMGSequencePlayer : public UObject, public IMovieScenePlayer
 {
 	GENERATED_UCLASS_BODY()
 
 public:
-	void InitSequencePlayer(UWidgetAnimation& InAnimation, UUserWidget& InUserWidget);
+	UMG_API void InitSequencePlayer(UWidgetAnimation& InAnimation, UUserWidget& InUserWidget);
 
 	/** Updates the running movie */
-	void Tick( float DeltaTime );
+	UMG_API void Tick( float DeltaTime );
 
 	/** Begins playing or restarts an animation */
-	void Play(float StartAtTime, int32 InNumLoopsToPlay, EUMGSequencePlayMode::Type InPlayMode, float InPlaybackSpeed, bool bRestoreState);
+	UMG_API void Play(float StartAtTime, int32 InNumLoopsToPlay, EUMGSequencePlayMode::Type InPlayMode, float InPlaybackSpeed, bool bRestoreState);
 
 	/** Begins playing or restarts an animation  and plays to the specified end time */
-	void PlayTo(float StartAtTime, float EndAtTime, int32 InNumLoopsToPlay, EUMGSequencePlayMode::Type InPlayMode, float InPlaybackSpeed, bool bRestoreState);
+	UMG_API void PlayTo(float StartAtTime, float EndAtTime, int32 InNumLoopsToPlay, EUMGSequencePlayMode::Type InPlayMode, float InPlaybackSpeed, bool bRestoreState);
 
 	/** Stops a running animation and resets time */
-	void Stop();
+	UMG_API void Stop();
 
 	/** Pauses a running animation */
-	void Pause();
+	UMG_API void Pause();
 
 	/** Reverses a running animation */
-	void Reverse();
+	UMG_API void Reverse();
 
 	void SetCurrentTime(float InTime) { TimeCursorPosition = AnimationResolution.AsFrameTime(InTime); }
 	FQualifiedFrameTime GetCurrentTime() const { return FQualifiedFrameTime(TimeCursorPosition, AnimationResolution); }
@@ -57,10 +57,10 @@ public:
 	}
 
 	/** Sets the number of loops to play */
-	void SetNumLoopsToPlay(int32 InNumLoopsToPlay);
+	UMG_API void SetNumLoopsToPlay(int32 InNumLoopsToPlay);
 
 	/** Sets the animation playback rate */
-	void SetPlaybackSpeed(float PlaybackSpeed);
+	UMG_API void SetPlaybackSpeed(float PlaybackSpeed);
 
 	/** Gets the current time position in the player (in seconds). */
 	bool IsPlayingForward() const { return bIsPlayingForward; }
@@ -70,40 +70,40 @@ public:
 
 	/** IMovieScenePlayer interface */
 	virtual FMovieSceneRootEvaluationTemplateInstance& GetEvaluationTemplate() override { return RootTemplateInstance; }
-	virtual UMovieSceneEntitySystemLinker* ConstructEntitySystemLinker() override;
+	UMG_API virtual UMovieSceneEntitySystemLinker* ConstructEntitySystemLinker() override;
 	virtual UObject* AsUObject() override { return this; }
 	virtual void UpdateCameraCut(UObject* CameraObject, const EMovieSceneCameraCutParams& CameraCutParams) override {}
 	virtual void SetViewportSettings(const TMap<FViewportClient*, EMovieSceneViewportParams>& ViewportParamsMap) override {}
 	virtual void GetViewportSettings(TMap<FViewportClient*, EMovieSceneViewportParams>& ViewportParamsMap) const override {}
-	virtual EMovieScenePlayerStatus::Type GetPlaybackStatus() const override;
-	virtual UObject* GetPlaybackContext() const override;
-	virtual TArray<UObject*> GetEventContexts() const override;
-	virtual void SetPlaybackStatus(EMovieScenePlayerStatus::Type InPlaybackStatus) override;
-	virtual void PopulateUpdateFlags(UE::MovieScene::ESequenceInstanceUpdateFlags& OutFlags) override;
+	UMG_API virtual EMovieScenePlayerStatus::Type GetPlaybackStatus() const override;
+	UMG_API virtual UObject* GetPlaybackContext() const override;
+	UMG_API virtual TArray<UObject*> GetEventContexts() const override;
+	UMG_API virtual void SetPlaybackStatus(EMovieScenePlayerStatus::Type InPlaybackStatus) override;
+	UMG_API virtual void PopulateUpdateFlags(UE::MovieScene::ESequenceInstanceUpdateFlags& OutFlags) override;
 
 	/** UObject interface */
-	virtual void BeginDestroy() override;
+	UMG_API virtual void BeginDestroy() override;
 
 	/** Disable this sequence player by removing any of its animation data from the entity manager */
-	void RemoveEvaluationData();
+	UMG_API void RemoveEvaluationData();
 
-	void TearDown();
+	UMG_API void TearDown();
 
 	DECLARE_EVENT_OneParam(UUMGSequencePlayer, FOnSequenceFinishedPlaying, UUMGSequencePlayer&);
 	FOnSequenceFinishedPlaying& OnSequenceFinishedPlaying() { return OnSequenceFinishedPlayingEvent; }
 
 private:
 	/** Internal play function with a verbose parameter set */
-	void PlayInternal(double StartAtTime, double EndAtTime, int32 InNumLoopsToPlay, EUMGSequencePlayMode::Type InPlayMode, float InPlaybackSpeed, bool bRestoreState);
+	UMG_API void PlayInternal(double StartAtTime, double EndAtTime, int32 InNumLoopsToPlay, EUMGSequencePlayMode::Type InPlayMode, float InPlaybackSpeed, bool bRestoreState);
 
 	/** Internal update function */
-	void UpdateInternal(FFrameTime LastTimePosition, FFrameTime NextTimePosition, bool bHasJumped);
+	UMG_API void UpdateInternal(FFrameTime LastTimePosition, FFrameTime NextTimePosition, bool bHasJumped);
 
-	bool NeedsQueueLatentAction() const;
-	void QueueLatentAction(FMovieSceneSequenceLatentActionDelegate Delegate);
-	void ApplyLatentActions();
+	UMG_API bool NeedsQueueLatentAction() const;
+	UMG_API void QueueLatentAction(FMovieSceneSequenceLatentActionDelegate Delegate);
+	UMG_API void ApplyLatentActions();
 
-	void HandleLatentStop();
+	UMG_API void HandleLatentStop();
 
 	/** Animation being played */
 	UPROPERTY()
diff --git a/Engine/Source/Runtime/UMG/Public/Animation/UMGSequenceTickManager.h b/Engine/Source/Runtime/UMG/Public/Animation/UMGSequenceTickManager.h
--- a/Engine/Source/Runtime/UMG/Public/Animation/UMGSequenceTickManager.h
+++ b/Engine/Source/Runtime/UMG/Public/Animation/UMGSequenceTickManager.h
@@ -34,35 +34,35 @@ struct FSequenceTickManagerWidgetData
 /**
  * An automatically created global object that will manage all widget animations.
  */
-UCLASS()
-class UMG_API UUMGSequenceTickManager : public UObject
+UCLASS(MinimalAPI)
+class UUMGSequenceTickManager : public UObject
 {
 public:
 	GENERATED_BODY()
 
-	UUMGSequenceTickManager(const FObjectInitializer& Init);
+	UMG_API UUMGSequenceTickManager(const FObjectInitializer& Init);
 
 	UMovieSceneEntitySystemLinker* GetLinker() { return Linker; }
 	TSharedPtr<FMovieSceneEntitySystemRunner> GetRunner() { return Runner; }
 
-	void AddLatentAction(FMovieSceneSequenceLatentActionDelegate Delegate);
-	void ClearLatentActions(UObject* Object);
-	void RunLatentActions();
+	UMG_API void AddLatentAction(FMovieSceneSequenceLatentActionDelegate Delegate);
+	UMG_API void ClearLatentActions(UObject* Object);
+	UMG_API void RunLatentActions();
 
-	static UUMGSequenceTickManager* Get(UObject* PlaybackContext);
+	static UMG_API UUMGSequenceTickManager* Get(UObject* PlaybackContext);
 
-	void ForceFlush();
+	UMG_API void ForceFlush();
 
-	void AddWidget(UUserWidget* InWidget);
-	void RemoveWidget(UUserWidget* InWidget);
+	UMG_API void AddWidget(UUserWidget* InWidget);
+	UMG_API void RemoveWidget(UUserWidget* InWidget);
 
-	void OnWidgetTicked(UUserWidget* InWidget);
+	UMG_API void OnWidgetTicked(UUserWidget* InWidget);
 
 private:
-	virtual void BeginDestroy() override;
+	UMG_API virtual void BeginDestroy() override;
 
-	void HandleSlatePostTick(float DeltaSeconds);
-	void TickWidgetAnimations(float DeltaSeconds);
+	UMG_API void HandleSlatePostTick(float DeltaSeconds);
+	UMG_API void TickWidgetAnimations(float DeltaSeconds);
 
 private:
 
diff --git a/Engine/Source/Runtime/UMG/Public/Animation/WidgetAnimation.h b/Engine/Source/Runtime/UMG/Public/Animation/WidgetAnimation.h
--- a/Engine/Source/Runtime/UMG/Public/Animation/WidgetAnimation.h
+++ b/Engine/Source/Runtime/UMG/Public/Animation/WidgetAnimation.h
@@ -33,7 +33,7 @@ public:
 	static UMG_API UWidgetAnimation* GetNullAnimation();
 
 	/** @return The friendly name of the animation */
-	UMG_API const FString& GetDisplayLabel() const
+	const FString& GetDisplayLabel() const
 	{
 		return DisplayLabel;
 	}
diff --git a/Engine/Source/Runtime/UMG/Public/Animation/WidgetAnimationDelegateBinding.h b/Engine/Source/Runtime/UMG/Public/Animation/WidgetAnimationDelegateBinding.h
--- a/Engine/Source/Runtime/UMG/Public/Animation/WidgetAnimationDelegateBinding.h
+++ b/Engine/Source/Runtime/UMG/Public/Animation/WidgetAnimationDelegateBinding.h
@@ -11,7 +11,7 @@
 class UUserWidget;
 
 USTRUCT()
-struct UMG_API FBlueprintWidgetAnimationDelegateBinding
+struct FBlueprintWidgetAnimationDelegateBinding
 {
 	GENERATED_USTRUCT_BODY()
 
@@ -36,13 +36,13 @@ struct UMG_API FBlueprintWidgetAnimationDelegateBinding
 	}
 };
 
-UCLASS()
-class UMG_API UWidgetAnimationDelegateBinding : public UDynamicBlueprintBinding
+UCLASS(MinimalAPI)
+class UWidgetAnimationDelegateBinding : public UDynamicBlueprintBinding
 {
 	GENERATED_UCLASS_BODY()
 
 	UPROPERTY()
 	TArray<FBlueprintWidgetAnimationDelegateBinding> WidgetAnimationDelegateBindings;
 
-	virtual void BindDynamicDelegates(UObject* InInstance) const override;
+	UMG_API virtual void BindDynamicDelegates(UObject* InInstance) const override;
 };
diff --git a/Engine/Source/Runtime/UMG/Public/Animation/WidgetMaterialTrackUtilities.h b/Engine/Source/Runtime/UMG/Public/Animation/WidgetMaterialTrackUtilities.h
--- a/Engine/Source/Runtime/UMG/Public/Animation/WidgetMaterialTrackUtilities.h
+++ b/Engine/Source/Runtime/UMG/Public/Animation/WidgetMaterialTrackUtilities.h
@@ -10,7 +10,7 @@ class UWidget;
 /**
  * Handle to a widget material for easy getting and setting of the material.  Not designed to be stored
  */
-class UMG_API FWidgetMaterialHandle
+class FWidgetMaterialHandle
 {
 public:
 	FWidgetMaterialHandle()
@@ -40,10 +40,10 @@ public:
 	bool IsValid() const { return TypeName != NAME_None && Data != nullptr; }
 
 	/** Gets the currently bound material */
-	UMaterialInterface* GetMaterial() const;
+	UMG_API UMaterialInterface* GetMaterial() const;
 
 	/** Sets the new bound material */
-	void SetMaterial(UMaterialInterface* InMaterial, UWidget* OwnerWidget);
+	UMG_API void SetMaterial(UMaterialInterface* InMaterial, UWidget* OwnerWidget);
 
 private:
 	/** Struct typename for that the data is pointing to */
diff --git a/Engine/Source/Runtime/UMG/Public/Binding/BoolBinding.h b/Engine/Source/Runtime/UMG/Public/Binding/BoolBinding.h
--- a/Engine/Source/Runtime/UMG/Public/Binding/BoolBinding.h
+++ b/Engine/Source/Runtime/UMG/Public/Binding/BoolBinding.h
@@ -8,18 +8,18 @@
 #include "Binding/PropertyBinding.h"
 #include "BoolBinding.generated.h"
 
-UCLASS()
-class UMG_API UBoolBinding : public UPropertyBinding
+UCLASS(MinimalAPI)
+class UBoolBinding : public UPropertyBinding
 {
 	GENERATED_BODY()
 
 public:
 
-	UBoolBinding();
+	UMG_API UBoolBinding();
 
-	virtual bool IsSupportedSource(FProperty* Property) const override;
-	virtual bool IsSupportedDestination(FProperty* Property) const override;
+	UMG_API virtual bool IsSupportedSource(FProperty* Property) const override;
+	UMG_API virtual bool IsSupportedDestination(FProperty* Property) const override;
 
 	UFUNCTION()
-	bool GetValue() const;
+	UMG_API bool GetValue() const;
 };
diff --git a/Engine/Source/Runtime/UMG/Public/Binding/BrushBinding.h b/Engine/Source/Runtime/UMG/Public/Binding/BrushBinding.h
--- a/Engine/Source/Runtime/UMG/Public/Binding/BrushBinding.h
+++ b/Engine/Source/Runtime/UMG/Public/Binding/BrushBinding.h
@@ -9,20 +9,20 @@
 #include "Binding/PropertyBinding.h"
 #include "BrushBinding.generated.h"
 
-UCLASS()
-class UMG_API UBrushBinding : public UPropertyBinding
+UCLASS(MinimalAPI)
+class UBrushBinding : public UPropertyBinding
 {
 	GENERATED_BODY()
 
 public:
 
-	UBrushBinding();
+	UMG_API UBrushBinding();
 
-	virtual bool IsSupportedSource(FProperty* Property) const override;
-	virtual bool IsSupportedDestination(FProperty* Property) const override;
+	UMG_API virtual bool IsSupportedSource(FProperty* Property) const override;
+	UMG_API virtual bool IsSupportedDestination(FProperty* Property) const override;
 
 	UFUNCTION()
-	FSlateBrush GetValue() const;
+	UMG_API FSlateBrush GetValue() const;
 
 private:
 	enum class EConversion : uint8
diff --git a/Engine/Source/Runtime/UMG/Public/Binding/CheckedStateBinding.h b/Engine/Source/Runtime/UMG/Public/Binding/CheckedStateBinding.h
--- a/Engine/Source/Runtime/UMG/Public/Binding/CheckedStateBinding.h
+++ b/Engine/Source/Runtime/UMG/Public/Binding/CheckedStateBinding.h
@@ -10,20 +10,20 @@
 
 enum class ECheckBoxState : uint8;
 
-UCLASS()
-class UMG_API UCheckedStateBinding : public UPropertyBinding
+UCLASS(MinimalAPI)
+class UCheckedStateBinding : public UPropertyBinding
 {
 	GENERATED_BODY()
 
 public:
 
-	UCheckedStateBinding();
+	UMG_API UCheckedStateBinding();
 
-	virtual bool IsSupportedSource(FProperty* Property) const override;
-	virtual bool IsSupportedDestination(FProperty* Property) const override;
+	UMG_API virtual bool IsSupportedSource(FProperty* Property) const override;
+	UMG_API virtual bool IsSupportedDestination(FProperty* Property) const override;
 
 	UFUNCTION()
-	ECheckBoxState GetValue() const;
+	UMG_API ECheckBoxState GetValue() const;
 
 private:
 	enum class EConversion : uint8
diff --git a/Engine/Source/Runtime/UMG/Public/Binding/ColorBinding.h b/Engine/Source/Runtime/UMG/Public/Binding/ColorBinding.h
--- a/Engine/Source/Runtime/UMG/Public/Binding/ColorBinding.h
+++ b/Engine/Source/Runtime/UMG/Public/Binding/ColorBinding.h
@@ -9,25 +9,25 @@
 #include "Binding/PropertyBinding.h"
 #include "ColorBinding.generated.h"
 
-UCLASS()
-class UMG_API UColorBinding : public UPropertyBinding
+UCLASS(MinimalAPI)
+class UColorBinding : public UPropertyBinding
 {
 	GENERATED_BODY()
 
 public:
 
-	UColorBinding();
+	UMG_API UColorBinding();
 
-	virtual bool IsSupportedSource(FProperty* Property) const override;
-	virtual bool IsSupportedDestination(FProperty* Property) const override;
+	UMG_API virtual bool IsSupportedSource(FProperty* Property) const override;
+	UMG_API virtual bool IsSupportedDestination(FProperty* Property) const override;
 
-	virtual void Bind(FProperty* Property, FScriptDelegate* Delegate) override;
+	UMG_API virtual void Bind(FProperty* Property, FScriptDelegate* Delegate) override;
 
 	UFUNCTION()
-	FSlateColor GetSlateValue() const;
+	UMG_API FSlateColor GetSlateValue() const;
 
 	UFUNCTION()
-	FLinearColor GetLinearValue() const;
+	UMG_API FLinearColor GetLinearValue() const;
 
 private:
 	mutable TOptional<bool> bNeedsConversion;
diff --git a/Engine/Source/Runtime/UMG/Public/Binding/DynamicPropertyPath.h b/Engine/Source/Runtime/UMG/Public/Binding/DynamicPropertyPath.h
--- a/Engine/Source/Runtime/UMG/Public/Binding/DynamicPropertyPath.h
+++ b/Engine/Source/Runtime/UMG/Public/Binding/DynamicPropertyPath.h
@@ -15,20 +15,20 @@ class UObject;
 
 /** */
 USTRUCT()
-struct UMG_API FDynamicPropertyPath : public FCachedPropertyPath
+struct FDynamicPropertyPath : public FCachedPropertyPath
 {
 	GENERATED_USTRUCT_BODY()
 
 public:
 
 	/** */
-	FDynamicPropertyPath();
+	UMG_API FDynamicPropertyPath();
 
 	/** */
-	FDynamicPropertyPath(const FString& Path);
+	UMG_API FDynamicPropertyPath(const FString& Path);
 
 	/** */
-	FDynamicPropertyPath(const TArray<FString>& PropertyChain);
+	UMG_API FDynamicPropertyPath(const TArray<FString>& PropertyChain);
 
 	/** Get the value represented by this property path */
 	template<typename T>
diff --git a/Engine/Source/Runtime/UMG/Public/Binding/FloatBinding.h b/Engine/Source/Runtime/UMG/Public/Binding/FloatBinding.h
--- a/Engine/Source/Runtime/UMG/Public/Binding/FloatBinding.h
+++ b/Engine/Source/Runtime/UMG/Public/Binding/FloatBinding.h
@@ -8,18 +8,18 @@
 #include "Binding/PropertyBinding.h"
 #include "FloatBinding.generated.h"
 
-UCLASS()
-class UMG_API UFloatBinding : public UPropertyBinding
+UCLASS(MinimalAPI)
+class UFloatBinding : public UPropertyBinding
 {
 	GENERATED_BODY()
 
 public:
 
-	UFloatBinding();
+	UMG_API UFloatBinding();
 
-	virtual bool IsSupportedSource(FProperty* Property) const override;
-	virtual bool IsSupportedDestination(FProperty* Property) const override;
+	UMG_API virtual bool IsSupportedSource(FProperty* Property) const override;
+	UMG_API virtual bool IsSupportedDestination(FProperty* Property) const override;
 
 	UFUNCTION()
-	float GetValue() const;
+	UMG_API float GetValue() const;
 };
diff --git a/Engine/Source/Runtime/UMG/Public/Binding/Int32Binding.h b/Engine/Source/Runtime/UMG/Public/Binding/Int32Binding.h
--- a/Engine/Source/Runtime/UMG/Public/Binding/Int32Binding.h
+++ b/Engine/Source/Runtime/UMG/Public/Binding/Int32Binding.h
@@ -8,18 +8,18 @@
 #include "Binding/PropertyBinding.h"
 #include "Int32Binding.generated.h"
 
-UCLASS()
-class UMG_API UInt32Binding : public UPropertyBinding
+UCLASS(MinimalAPI)
+class UInt32Binding : public UPropertyBinding
 {
 	GENERATED_BODY()
 
 public:
 
-	UInt32Binding();
+	UMG_API UInt32Binding();
 
-	virtual bool IsSupportedSource(FProperty* Property) const override;
-	virtual bool IsSupportedDestination(FProperty* Property) const override;
+	UMG_API virtual bool IsSupportedSource(FProperty* Property) const override;
+	UMG_API virtual bool IsSupportedDestination(FProperty* Property) const override;
 
 	UFUNCTION()
-	int32 GetValue() const;
+	UMG_API int32 GetValue() const;
 };
diff --git a/Engine/Source/Runtime/UMG/Public/Binding/MouseCursorBinding.h b/Engine/Source/Runtime/UMG/Public/Binding/MouseCursorBinding.h
--- a/Engine/Source/Runtime/UMG/Public/Binding/MouseCursorBinding.h
+++ b/Engine/Source/Runtime/UMG/Public/Binding/MouseCursorBinding.h
@@ -9,18 +9,18 @@
 #include "Binding/PropertyBinding.h"
 #include "MouseCursorBinding.generated.h"
 
-UCLASS()
-class UMG_API UMouseCursorBinding : public UPropertyBinding
+UCLASS(MinimalAPI)
+class UMouseCursorBinding : public UPropertyBinding
 {
 	GENERATED_BODY()
 
 public:
 
-	UMouseCursorBinding();
+	UMG_API UMouseCursorBinding();
 
-	virtual bool IsSupportedSource(FProperty* Property) const override;
-	virtual bool IsSupportedDestination(FProperty* Property) const override;
+	UMG_API virtual bool IsSupportedSource(FProperty* Property) const override;
+	UMG_API virtual bool IsSupportedDestination(FProperty* Property) const override;
 
 	UFUNCTION()
-	EMouseCursor::Type GetValue() const;
+	UMG_API EMouseCursor::Type GetValue() const;
 };
diff --git a/Engine/Source/Runtime/UMG/Public/Binding/PropertyBinding.h b/Engine/Source/Runtime/UMG/Public/Binding/PropertyBinding.h
--- a/Engine/Source/Runtime/UMG/Public/Binding/PropertyBinding.h
+++ b/Engine/Source/Runtime/UMG/Public/Binding/PropertyBinding.h
@@ -13,18 +13,18 @@
 
 DECLARE_CYCLE_STAT_EXTERN(TEXT("UMG Binding"), STAT_UMGBinding, STATGROUP_Slate,);
 
-UCLASS()
-class UMG_API UPropertyBinding : public UObject
+UCLASS(MinimalAPI)
+class UPropertyBinding : public UObject
 {
 	GENERATED_BODY()
 
 public:
-	UPropertyBinding();
+	UMG_API UPropertyBinding();
 
-	virtual bool IsSupportedSource(FProperty* Property) const;
-	virtual bool IsSupportedDestination(FProperty* Property) const;
+	UMG_API virtual bool IsSupportedSource(FProperty* Property) const;
+	UMG_API virtual bool IsSupportedDestination(FProperty* Property) const;
 
-	virtual void Bind(FProperty* Property, FScriptDelegate* Delegate);
+	UMG_API virtual void Bind(FProperty* Property, FScriptDelegate* Delegate);
 
 public:
 	/** The source object to use as the initial container to resolve the Source Property Path on. */
diff --git a/Engine/Source/Runtime/UMG/Public/Binding/States/WidgetStateBitfield.h b/Engine/Source/Runtime/UMG/Public/Binding/States/WidgetStateBitfield.h
--- a/Engine/Source/Runtime/UMG/Public/Binding/States/WidgetStateBitfield.h
+++ b/Engine/Source/Runtime/UMG/Public/Binding/States/WidgetStateBitfield.h
@@ -19,27 +19,27 @@
  * bitfield ahead of time & store that to later be compared against  
  */
 USTRUCT()
-struct UMG_API FWidgetStateBitfield
+struct FWidgetStateBitfield
 {
 	GENERATED_BODY()
 
 public:
 
 	/** Default, will convert to false as a bool */
-	FWidgetStateBitfield();
+	UMG_API FWidgetStateBitfield();
 
 	/** Interprets name as binary state on */
-	FWidgetStateBitfield(const FName InStateName);
+	UMG_API FWidgetStateBitfield(const FName InStateName);
 
 	/** Interprets name value pair as enum state */
-	FWidgetStateBitfield(const FName InStateName, const uint8 InValue);
+	UMG_API FWidgetStateBitfield(const FName InStateName, const uint8 InValue);
 
 public:
 
 	/** 
 	 * Negation of binary and enum states, maintains enum state usage, but marks states as not allowed.
 	 */
-	FWidgetStateBitfield operator~() const;
+	UMG_API FWidgetStateBitfield operator~() const;
 
 	/** 
 	 * Intersection of binary and enum states, with union on enum state usage so enum state mismatches can fail
@@ -47,28 +47,28 @@ public:
 	 * Approximation of bitwise '&'. Non-commutative, will use LHS enum state set if RHS does not use enum. Additionally,
 	 * if the RHS has binary states & the LHS does not meet those, all enum states will be marked unused to ensure failure
 	 */
-	FWidgetStateBitfield Intersect(const FWidgetStateBitfield& Rhs) const;
+	UMG_API FWidgetStateBitfield Intersect(const FWidgetStateBitfield& Rhs) const;
 
 	/** 
 	 * Union of binary and enum states, allows for a single enum state to have multiple values in a bitfield.
 	 * Unions will grow enum state set if allowance is same, otherwise the allowed set is maintained.
 	 * Approximation of bitwise '|'. Commutative.
 	 */
-	FWidgetStateBitfield Union(const FWidgetStateBitfield& Rhs) const;
+	UMG_API FWidgetStateBitfield Union(const FWidgetStateBitfield& Rhs) const;
 
 public:
 
 	/** True if any binary or enum state is set, and no empty used enum states exist */
-	operator bool() const;
+	UMG_API operator bool() const;
 
 	/** True if any binary state is set */
-	bool HasBinaryStates() const;
+	UMG_API bool HasBinaryStates() const;
 
 	/** True if any enum state is used, does not indicate if state is actually set. Use 'HasEmptyUsedEnumStates' for that. */
-	bool HasEnumStates() const;
+	UMG_API bool HasEnumStates() const;
 
 	/** True if any enum used enum states are empty, usually indicating a failed '&' test */
-	bool HasEmptyUsedEnumStates() const;
+	UMG_API bool HasEmptyUsedEnumStates() const;
 
 public:
 
@@ -77,30 +77,30 @@ public:
 	 */
 
 	/** True if any state flag is met */
-	bool HasAnyFlags(const FWidgetStateBitfield& InBitfield) const;
+	UMG_API bool HasAnyFlags(const FWidgetStateBitfield& InBitfield) const;
 
 	/** True if all state flags are met */
-	bool HasAllFlags(const FWidgetStateBitfield& InBitfield) const;
+	UMG_API bool HasAllFlags(const FWidgetStateBitfield& InBitfield) const;
 
 	/** True if any binary state flag is met */
-	bool HasAnyBinaryFlags(const FWidgetStateBitfield& InBitfield) const;
+	UMG_API bool HasAnyBinaryFlags(const FWidgetStateBitfield& InBitfield) const;
 
 	/** True if all binary state flags are met */
-	bool HasAllBinaryFlags(const FWidgetStateBitfield& InBitfield) const;
+	UMG_API bool HasAllBinaryFlags(const FWidgetStateBitfield& InBitfield) const;
 
 	/** True if any enum state flag is met */
-	bool HasAnyEnumFlags(const FWidgetStateBitfield& InBitfield) const;
+	UMG_API bool HasAnyEnumFlags(const FWidgetStateBitfield& InBitfield) const;
 
 	/** True if all enum state flags are met */
-	bool HasAllEnumFlags(const FWidgetStateBitfield& InBitfield) const;
+	UMG_API bool HasAllEnumFlags(const FWidgetStateBitfield& InBitfield) const;
 
 public:
 
 	/** Set state to given value, just calls assignment operator for you */
-	void SetState(const FWidgetStateBitfield& InBitfield);
+	UMG_API void SetState(const FWidgetStateBitfield& InBitfield);
 
 	/** Negate both binary and enum states */
-	void NegateStates();
+	UMG_API void NegateStates();
 
 	/**
 	 * Set binary state to given value
@@ -108,7 +108,7 @@ public:
 	 * @param BinaryStateIndex Index for Binary state to assign value to
 	 * @param BinaryStateValue Value to assign for given Binary state
 	 */
-	void SetBinaryState(uint8 BinaryStateIndex, bool BinaryStateValue);
+	UMG_API void SetBinaryState(uint8 BinaryStateIndex, bool BinaryStateValue);
 
 	/**
 	 * Set Binary state to given value, using a bitfield to indicate states that should be changed
@@ -116,13 +116,13 @@ public:
 	 * @param BinaryStateBitfield state bitfield, all marked states will be set with given value
 	 * @param BinaryStateValue Value to assign for given Binary state
 	 */
-	void SetBinaryState(const FWidgetStateBitfield& BinaryStateBitfield, bool BinaryStateValue);
+	UMG_API void SetBinaryState(const FWidgetStateBitfield& BinaryStateBitfield, bool BinaryStateValue);
 
 	/** Set Binary state to given value, slow version w/ FName map lookup */
-	void SetBinaryStateSlow(FName BinaryStateName, bool BinaryStateValue);
+	UMG_API void SetBinaryStateSlow(FName BinaryStateName, bool BinaryStateValue);
 
 	/** Negate only binary states */
-	void NegateBinaryStates();
+	UMG_API void NegateBinaryStates();
 
 	/** 
 	 * Set enum state to given value, clearing all existing values for that state  
@@ -130,37 +130,37 @@ public:
 	 * @param EnumStateIndex Index for enum state to assign value to
 	 * @param EnumStateValue Value to assign for given enum state
 	 */
-	void SetEnumState(uint8 EnumStateIndex, uint8 EnumStateValue);
+	UMG_API void SetEnumState(uint8 EnumStateIndex, uint8 EnumStateValue);
 
 	/**
 	 * Set enum state to given value, clearing all existing values for states in use
 	 *
 	 * @param EnumStateBitfield state bitfield, all used states will copy-overwrite current enum state
 	 */
-	void SetEnumState(const FWidgetStateBitfield& EnumStateBitfield);
+	UMG_API void SetEnumState(const FWidgetStateBitfield& EnumStateBitfield);
 
 	/** Set enum state to given value, clearing all existing values for that state, slow version w/ FName map lookup */
-	void SetEnumStateSlow(FName EnumStateName, uint8 EnumStateValue);
+	UMG_API void SetEnumStateSlow(FName EnumStateName, uint8 EnumStateValue);
 
 	/**
 	 * Clear all existing values for given states & mark the states unused
 	 *
 	 * @param EnumStateBitfield state bitfield, all used states will be cleared
 	 */
-	void ClearEnumState(const FWidgetStateBitfield& EnumStateBitfield);
+	UMG_API void ClearEnumState(const FWidgetStateBitfield& EnumStateBitfield);
 
 	/**
 	 * Clear all existing values for given state & mark the state unused
 	 *
 	 * @param EnumStateIndex Index for enum state to mark as unused
 	 */
-	void ClearEnumState(uint8 EnumStateIndex);
+	UMG_API void ClearEnumState(uint8 EnumStateIndex);
 
 	/** Clear all existing values for given state & mark the state unused, slow version w/ map lookup */
-	void ClearEnumState(FName EnumStateName);
+	UMG_API void ClearEnumState(FName EnumStateName);
 
 	/** Negate only enum states, does not change usage, but marks existing values as not allowed */
-	void NegateEnumStates();
+	UMG_API void NegateEnumStates();
 
 private:
 
@@ -191,4 +191,4 @@ private:
 
 	/** Binary states captured by this state field. Note: 16 == sizeof(IsEnumStateUsed) * CHAR_BIT */
 	TStaticArray<TSet<uint8, DefaultKeyFuncs<uint8, false>, TInlineSetAllocator<1>>, 16> EnumStates;
-};
\ No newline at end of file
+};
diff --git a/Engine/Source/Runtime/UMG/Public/Binding/States/WidgetStateRegistration.h b/Engine/Source/Runtime/UMG/Public/Binding/States/WidgetStateRegistration.h
--- a/Engine/Source/Runtime/UMG/Public/Binding/States/WidgetStateRegistration.h
+++ b/Engine/Source/Runtime/UMG/Public/Binding/States/WidgetStateRegistration.h
@@ -18,8 +18,8 @@ class UWidgetStateSettings;
  * Technically these can be created in BP, but for now we don't want to encourage
  * that workflow as it involves requring overrides for the virtuals which is technical.
  */
-UCLASS(Transient)
-class UMG_API UWidgetBinaryStateRegistration : public UObject
+UCLASS(Transient, MinimalAPI)
+class UWidgetBinaryStateRegistration : public UObject
 {
 	GENERATED_BODY()
 
@@ -28,123 +28,123 @@ public:
 	virtual ~UWidgetBinaryStateRegistration() = default;
 
 	/** Called once during WidgetStateSettings initialization to get this widget state's name */
-	virtual FName GetStateName() const;
+	UMG_API virtual FName GetStateName() const;
 
 	/** Called on widget registration to correctly initialize widget state based on the current widget */
-	virtual bool GetRegisteredWidgetState(const UWidget* InWidget) const;
+	UMG_API virtual bool GetRegisteredWidgetState(const UWidget* InWidget) const;
 
 protected:
 	friend UWidgetStateSettings;
 
 	/** Called to give CDO chance to initialize any static state bitfields that might be declared for convenience */
-	virtual void InitializeStaticBitfields() const;
+	UMG_API virtual void InitializeStaticBitfields() const;
 };
 
-UCLASS(Transient)
-class UMG_API UWidgetHoveredStateRegistration : public UWidgetBinaryStateRegistration
+UCLASS(Transient, MinimalAPI)
+class UWidgetHoveredStateRegistration : public UWidgetBinaryStateRegistration
 {
 	GENERATED_BODY()
 
 public:
 
 	/** Post-load initialized bit corresponding to this binary state */
-	static inline FWidgetStateBitfield Bit;
+	static UMG_API inline FWidgetStateBitfield Bit;
 
 	static const inline FName StateName = FName("Hovered");
 
 	//~ Begin UWidgetBinaryStateRegistration Interface.
-	virtual FName GetStateName() const override;
-	virtual bool GetRegisteredWidgetState(const UWidget* InWidget) const override;
+	UMG_API virtual FName GetStateName() const override;
+	UMG_API virtual bool GetRegisteredWidgetState(const UWidget* InWidget) const override;
 	//~ End UWidgetBinaryStateRegistration Interface
 
 protected:
 	friend UWidgetStateSettings;
 
 	//~ Begin UWidgetBinaryStateRegistration Interface.
-	virtual void InitializeStaticBitfields() const override;
+	UMG_API virtual void InitializeStaticBitfields() const override;
 	//~ End UWidgetBinaryStateRegistration Interface
 };
 
-UCLASS(Transient)
-class UMG_API UWidgetPressedStateRegistration : public UWidgetBinaryStateRegistration
+UCLASS(Transient, MinimalAPI)
+class UWidgetPressedStateRegistration : public UWidgetBinaryStateRegistration
 {
 	GENERATED_BODY()
 
 public:
 
 	/** Post-load initialized bit corresponding to this binary state */
-	static inline FWidgetStateBitfield Bit;
+	static UMG_API inline FWidgetStateBitfield Bit;
 
 	static const inline FName StateName = FName("Pressed");
 
 	//~ Begin UWidgetBinaryStateRegistration Interface.
-	virtual FName GetStateName() const override;
-	virtual bool GetRegisteredWidgetState(const UWidget* InWidget) const override;
+	UMG_API virtual FName GetStateName() const override;
+	UMG_API virtual bool GetRegisteredWidgetState(const UWidget* InWidget) const override;
 	//~ End UWidgetBinaryStateRegistration Interface
 
 protected:
 	friend UWidgetStateSettings;
 
 	//~ Begin UWidgetBinaryStateRegistration Interface.
-	virtual void InitializeStaticBitfields() const override;
+	UMG_API virtual void InitializeStaticBitfields() const override;
 	//~ End UWidgetBinaryStateRegistration Interface
 };
 
-UCLASS(Transient)
-class UMG_API UWidgetDisabledStateRegistration : public UWidgetBinaryStateRegistration
+UCLASS(Transient, MinimalAPI)
+class UWidgetDisabledStateRegistration : public UWidgetBinaryStateRegistration
 {
 	GENERATED_BODY()
 
 public:
 
 	/** Post-load initialized bit corresponding to this binary state */
-	static inline FWidgetStateBitfield Bit;
+	static UMG_API inline FWidgetStateBitfield Bit;
 
 	static const inline FName StateName = FName("Disabled");
 
 	//~ Begin UWidgetBinaryStateRegistration Interface.
-	virtual FName GetStateName() const override;
-	virtual bool GetRegisteredWidgetState(const UWidget* InWidget) const override;
+	UMG_API virtual FName GetStateName() const override;
+	UMG_API virtual bool GetRegisteredWidgetState(const UWidget* InWidget) const override;
 	//~ End UWidgetBinaryStateRegistration Interface
 
 protected:
 	friend UWidgetStateSettings;
 
 	//~ Begin UWidgetBinaryStateRegistration Interface.
-	virtual void InitializeStaticBitfields() const override;
+	UMG_API virtual void InitializeStaticBitfields() const override;
 	//~ End UWidgetBinaryStateRegistration Interface
 };
 
-UCLASS(Transient)
-class UMG_API UWidgetSelectedStateRegistration : public UWidgetBinaryStateRegistration
+UCLASS(Transient, MinimalAPI)
+class UWidgetSelectedStateRegistration : public UWidgetBinaryStateRegistration
 {
 	GENERATED_BODY()
 
 public:
 
 	/** Post-load initialized bit corresponding to this binary state */
-	static inline FWidgetStateBitfield Bit;
+	static UMG_API inline FWidgetStateBitfield Bit;
 
 	static const inline FName StateName = FName("Selected");
 
 	//~ Begin UWidgetBinaryStateRegistration Interface.
-	virtual FName GetStateName() const override;
-	virtual bool GetRegisteredWidgetState(const UWidget* InWidget) const override;
+	UMG_API virtual FName GetStateName() const override;
+	UMG_API virtual bool GetRegisteredWidgetState(const UWidget* InWidget) const override;
 	//~ End UWidgetBinaryStateRegistration Interface
 
 protected:
 	friend UWidgetStateSettings;
 
 	//~ Begin UWidgetBinaryStateRegistration Interface.
-	virtual void InitializeStaticBitfields() const override;
+	UMG_API virtual void InitializeStaticBitfields() const override;
 	//~ End UWidgetBinaryStateRegistration Interface
 };
 
 /**
  * Derive from to add a new Enum binary state
  */
-UCLASS(Transient)
-class UMG_API UWidgetEnumStateRegistration : public UObject
+UCLASS(Transient, MinimalAPI)
+class UWidgetEnumStateRegistration : public UObject
 {
 	GENERATED_BODY()
 
@@ -153,17 +153,17 @@ public:
 	virtual ~UWidgetEnumStateRegistration() = default;
 
 	/** Called once during WidgetStateSettings initialization to get this widget state's name */
-	virtual FName GetStateName() const;
+	UMG_API virtual FName GetStateName() const;
 
 	/** Called on widget registration to determine if this widget uses the given state */
-	virtual bool GetRegisteredWidgetUsesState(const UWidget* InWidget) const;
+	UMG_API virtual bool GetRegisteredWidgetUsesState(const UWidget* InWidget) const;
 
 	/** Called on widget registration to correctly initialize widget state based on the current widget */
-	virtual uint8 GetRegisteredWidgetState(const UWidget* InWidget) const;
+	UMG_API virtual uint8 GetRegisteredWidgetState(const UWidget* InWidget) const;
 
 protected:
 	friend UWidgetStateSettings;
 
 	/** Called to give CDO chance to initialize any static state bitfields that might be declared for convenience */
-	virtual void InitializeStaticBitfields() const;
+	UMG_API virtual void InitializeStaticBitfields() const;
 };
diff --git a/Engine/Source/Runtime/UMG/Public/Binding/States/WidgetStateSettings.h b/Engine/Source/Runtime/UMG/Public/Binding/States/WidgetStateSettings.h
--- a/Engine/Source/Runtime/UMG/Public/Binding/States/WidgetStateSettings.h
+++ b/Engine/Source/Runtime/UMG/Public/Binding/States/WidgetStateSettings.h
@@ -22,8 +22,8 @@ class UWidgetEnumStateRegistration;
  * 
  * Note: Currently doesn't really have any meaningful settings, used more as global singleton. May change.
  */
-UCLASS(config = Game, defaultconfig)
-class UMG_API UWidgetStateSettings : public UDeveloperSettings
+UCLASS(config = Game, defaultconfig, MinimalAPI)
+class UWidgetStateSettings : public UDeveloperSettings
 {
 	GENERATED_BODY()
 
@@ -35,34 +35,34 @@ public:
 	}
 
 	//~ Begin UObject Interface.
-	virtual void PostInitProperties() override;
+	UMG_API virtual void PostInitProperties() override;
 	//~ End UObject Interface
 
 public:
 
 	/** Obtain a list of all known States. */
-	void GetAllStateNames(TArray<FName>& OutStateNames) const;
+	UMG_API void GetAllStateNames(TArray<FName>& OutStateNames) const;
 
 	/** Obtain a list of all known Binary States. */
-	void GetBinaryStateNames(TArray<FName>& OutBinaryStateNames) const;
+	UMG_API void GetBinaryStateNames(TArray<FName>& OutBinaryStateNames) const;
 
 	/** Obtain a list of all known Enum States. */
-	void GetEnumStateNames(TArray<FName>& OutEnumStateNames) const;
+	UMG_API void GetEnumStateNames(TArray<FName>& OutEnumStateNames) const;
 
 	/** Get Index corresponding to this binary state, which may dynamically change - but not within a game session. */
-	uint8 GetBinaryStateIndex(const FName InBinaryStateName) const;
+	UMG_API uint8 GetBinaryStateIndex(const FName InBinaryStateName) const;
 
 	/** Get Index corresponding to this enum state, which may dynamically change - but not within a game session. */
-	uint8 GetEnumStateIndex(const FName InEnumStateName) const;
+	UMG_API uint8 GetEnumStateIndex(const FName InEnumStateName) const;
 
 	/** Get Name corresponding to this binary state index, which may dynamically change - but not within a game session. */
-	FName GetBinaryStateName(const uint8 InBinaryStateIndex) const;
+	UMG_API FName GetBinaryStateName(const uint8 InBinaryStateIndex) const;
 
 	/** Get Name corresponding to this enum state index, which may dynamically change - but not within a game session. */
-	FName GetEnumStateName(const uint8 InEnumStateIndex) const;
+	UMG_API FName GetEnumStateName(const uint8 InEnumStateIndex) const;
 
 	/** Gets bitfield of widget based on current state and given on registration state initializers */
-	FWidgetStateBitfield GetInitialRegistrationBitfield(const UWidget* InWidget) const;
+	UMG_API FWidgetStateBitfield GetInitialRegistrationBitfield(const UWidget* InWidget) const;
 
 private:
 	/** Map of Binary State Names to index (In order of addition) */
@@ -82,4 +82,4 @@ private:
 
 	/** Ordered list of Enum Widget States */
 	TArray<TObjectPtr<UWidgetEnumStateRegistration>, TInlineAllocator<8>> EnumStateRegistrationCDOs;
-};
\ No newline at end of file
+};
diff --git a/Engine/Source/Runtime/UMG/Public/Binding/TextBinding.h b/Engine/Source/Runtime/UMG/Public/Binding/TextBinding.h
--- a/Engine/Source/Runtime/UMG/Public/Binding/TextBinding.h
+++ b/Engine/Source/Runtime/UMG/Public/Binding/TextBinding.h
@@ -8,23 +8,23 @@
 #include "Binding/PropertyBinding.h"
 #include "TextBinding.generated.h"
 
-UCLASS()
-class UMG_API UTextBinding : public UPropertyBinding
+UCLASS(MinimalAPI)
+class UTextBinding : public UPropertyBinding
 {
 	GENERATED_BODY()
 
 public:
 
-	virtual bool IsSupportedSource(FProperty* Property) const override;
-	virtual bool IsSupportedDestination(FProperty* Property) const override;
+	UMG_API virtual bool IsSupportedSource(FProperty* Property) const override;
+	UMG_API virtual bool IsSupportedDestination(FProperty* Property) const override;
 
-	virtual void Bind(FProperty* Property, FScriptDelegate* Delegate) override;
+	UMG_API virtual void Bind(FProperty* Property, FScriptDelegate* Delegate) override;
 
 	UFUNCTION()
-	FText GetTextValue() const;
+	UMG_API FText GetTextValue() const;
 
 	UFUNCTION()
-	FString GetStringValue() const;
+	UMG_API FString GetStringValue() const;
 
 private:
 
diff --git a/Engine/Source/Runtime/UMG/Public/Binding/VisibilityBinding.h b/Engine/Source/Runtime/UMG/Public/Binding/VisibilityBinding.h
--- a/Engine/Source/Runtime/UMG/Public/Binding/VisibilityBinding.h
+++ b/Engine/Source/Runtime/UMG/Public/Binding/VisibilityBinding.h
@@ -11,18 +11,18 @@
 
 enum class ESlateVisibility : uint8;
 
-UCLASS()
-class UMG_API UVisibilityBinding : public UPropertyBinding
+UCLASS(MinimalAPI)
+class UVisibilityBinding : public UPropertyBinding
 {
 	GENERATED_BODY()
 
 public:
 
-	UVisibilityBinding();
+	UMG_API UVisibilityBinding();
 
-	virtual bool IsSupportedSource(FProperty* Property) const override;
-	virtual bool IsSupportedDestination(FProperty* Property) const override;
+	UMG_API virtual bool IsSupportedSource(FProperty* Property) const override;
+	UMG_API virtual bool IsSupportedDestination(FProperty* Property) const override;
 
 	UFUNCTION()
-	ESlateVisibility GetValue() const;
+	UMG_API ESlateVisibility GetValue() const;
 };
diff --git a/Engine/Source/Runtime/UMG/Public/Binding/WidgetBinding.h b/Engine/Source/Runtime/UMG/Public/Binding/WidgetBinding.h
--- a/Engine/Source/Runtime/UMG/Public/Binding/WidgetBinding.h
+++ b/Engine/Source/Runtime/UMG/Public/Binding/WidgetBinding.h
@@ -10,18 +10,18 @@
 
 class UWidget;
 
-UCLASS()
-class UMG_API UWidgetBinding : public UPropertyBinding
+UCLASS(MinimalAPI)
+class UWidgetBinding : public UPropertyBinding
 {
 	GENERATED_BODY()
 
 public:
 
-	UWidgetBinding();
+	UMG_API UWidgetBinding();
 
-	virtual bool IsSupportedSource(FProperty* Property) const override;
-	virtual bool IsSupportedDestination(FProperty* Property) const override;
+	UMG_API virtual bool IsSupportedSource(FProperty* Property) const override;
+	UMG_API virtual bool IsSupportedDestination(FProperty* Property) const override;
 
 	UFUNCTION()
-	UWidget* GetValue() const;
+	UMG_API UWidget* GetValue() const;
 };
diff --git a/Engine/Source/Runtime/UMG/Public/Blueprint/AsyncTaskDownloadImage.h b/Engine/Source/Runtime/UMG/Public/Blueprint/AsyncTaskDownloadImage.h
--- a/Engine/Source/Runtime/UMG/Public/Blueprint/AsyncTaskDownloadImage.h
+++ b/Engine/Source/Runtime/UMG/Public/Blueprint/AsyncTaskDownloadImage.h
@@ -13,14 +13,14 @@ class UTexture2DDynamic;
 
 DECLARE_DYNAMIC_MULTICAST_DELEGATE_OneParam(FDownloadImageDelegate, UTexture2DDynamic*, Texture);
 
-UCLASS()
-class UMG_API UAsyncTaskDownloadImage : public UBlueprintAsyncActionBase
+UCLASS(MinimalAPI)
+class UAsyncTaskDownloadImage : public UBlueprintAsyncActionBase
 {
 	GENERATED_UCLASS_BODY()
 
 public:
 	UFUNCTION(BlueprintCallable, meta=( BlueprintInternalUseOnly="true" ))
-	static UAsyncTaskDownloadImage* DownloadImage(FString URL);
+	static UMG_API UAsyncTaskDownloadImage* DownloadImage(FString URL);
 
 public:
 
@@ -32,7 +32,7 @@ public:
 
 public:
 
-	void Start(FString URL);
+	UMG_API void Start(FString URL);
 
 private:
 
diff --git a/Engine/Source/Runtime/UMG/Public/Blueprint/DragDropOperation.h b/Engine/Source/Runtime/UMG/Public/Blueprint/DragDropOperation.h
--- a/Engine/Source/Runtime/UMG/Public/Blueprint/DragDropOperation.h
+++ b/Engine/Source/Runtime/UMG/Public/Blueprint/DragDropOperation.h
@@ -37,8 +37,8 @@ enum class EDragPivot : uint8
 /**
  * This class is the base drag drop operation for UMG, extend it to add additional data and add new functionality.
  */
-UCLASS(BlueprintType, Blueprintable, meta=( DontUseGenericSpawnObject="True" ))
-class UMG_API UDragDropOperation : public UObject
+UCLASS(BlueprintType, Blueprintable, meta=( DontUseGenericSpawnObject="True" ), MinimalAPI)
+class UDragDropOperation : public UObject
 {
 	GENERATED_UCLASS_BODY()
 	
@@ -86,13 +86,13 @@ public:
 
 	/**  */
 	UFUNCTION(BlueprintNativeEvent, Category="Drag and Drop")
-	void Drop(const FPointerEvent& PointerEvent);
+	UMG_API void Drop(const FPointerEvent& PointerEvent);
 
 	/**  */
 	UFUNCTION(BlueprintNativeEvent, Category="Drag and Drop")
-	void DragCancelled(const FPointerEvent& PointerEvent);
+	UMG_API void DragCancelled(const FPointerEvent& PointerEvent);
 
 	/**  */
 	UFUNCTION(BlueprintNativeEvent, Category="Drag and Drop")
-	void Dragged(const FPointerEvent& PointerEvent);
+	UMG_API void Dragged(const FPointerEvent& PointerEvent);
 };
