[Editor Modes] Exposed some additional API on the mode manager for better typed element support

diff --git a/Engine/Source/Editor/UnrealEd/Private/EditorModeManager.cpp b/Engine/Source/Editor/UnrealEd/Private/EditorModeManager.cpp
--- a/Engine/Source/Editor/UnrealEd/Private/EditorModeManager.cpp
+++ b/Engine/Source/Editor/UnrealEd/Private/EditorModeManager.cpp
@@ -262,6 +262,11 @@ bool FEditorModeTools::SelectionHasSceneComponent() const
 	return bSelectionHasSceneComponent;
 }
 
+void FEditorModeTools::SetSelectionHasSceneComponent(bool bHasSceneComponent)
+{
+	bSelectionHasSceneComponent = bHasSceneComponent;
+}
+
 bool FEditorModeTools::IsSelectionAllowed(AActor* InActor, const bool bInSelected) const
 {
 	bool bSelectionAllowed = (ActiveScriptableModes.Num() == 0);
diff --git a/Engine/Source/Editor/UnrealEd/Public/EditorModeManager.h b/Engine/Source/Editor/UnrealEd/Public/EditorModeManager.h
--- a/Engine/Source/Editor/UnrealEd/Public/EditorModeManager.h
+++ b/Engine/Source/Editor/UnrealEd/Public/EditorModeManager.h
@@ -460,7 +460,7 @@ public:
 	 * Returns the selection set for the toolkit host.
 	 * (i.e. the selection set for the level editor)
 	 */
-	UNREALED_API UTypedElementSelectionSet* GetEditorSelectionSet() const;
+	UNREALED_API virtual UTypedElementSelectionSet* GetEditorSelectionSet() const;
 
 	/**
 	 * Stores the current selection under the given key, and clears the current selection state if requested.
@@ -491,6 +491,7 @@ public:
 	 * Whether or not the current selection has a scene component selected
  	 */
 	UNREALED_API bool SelectionHasSceneComponent() const;
+	UNREALED_API void SetSelectionHasSceneComponent(bool bHasSceneComponent);
 
 	UNREALED_API bool IsSelectionAllowed(AActor* InActor, const bool bInSelected) const;
 
diff --git a/Engine/Source/Runtime/TypedElementRuntime/Private/Elements/Framework/TypedElementSelectionSet.cpp b/Engine/Source/Runtime/TypedElementRuntime/Private/Elements/Framework/TypedElementSelectionSet.cpp
--- a/Engine/Source/Runtime/TypedElementRuntime/Private/Elements/Framework/TypedElementSelectionSet.cpp
+++ b/Engine/Source/Runtime/TypedElementRuntime/Private/Elements/Framework/TypedElementSelectionSet.cpp
@@ -388,6 +388,11 @@ bool UTypedElementSelectionSet::AllowSelectionModifiers(const FTypedElementHandl
 	return SelectionSetElement && SelectionSetElement.AllowSelectionModifiers();
 }
 
+void UTypedElementSelectionSet::SetNameForTedsIntegration(FName InNameForIntegration)
+{
+	ListNameForTedsIntegration = InNameForIntegration;
+}
+
 FTypedElementHandle UTypedElementSelectionSet::GetSelectionElement(const FTypedElementHandle& InElementHandle, const ETypedElementSelectionMethod InSelectionMethod) const
 {
 	FTypedElementSelectionSetElement SelectionSetElement = ResolveSelectionSetElement(InElementHandle);
@@ -466,7 +471,9 @@ void UTypedElementSelectionSet::RestoreSelectionState(const FTypedElementSelecti
 				.SetAllowHidden(true)
 				.SetAllowGroups(false)
 				.SetAllowLegacyNotifications(false)
-				.SetWarnIfLocked(false);
+				.SetWarnIfLocked(false)
+				.SetNameForTEDSIntegration(ListNameForTedsIntegration);
+
 
 			// TODO: Work out the intersection of the before and after state instead of clearing and reselecting?
 			SetSelection(SelectedElements, SelectionOptions);
diff --git a/Engine/Source/Runtime/TypedElementRuntime/Public/Elements/Framework/TypedElementSelectionSet.h b/Engine/Source/Runtime/TypedElementRuntime/Public/Elements/Framework/TypedElementSelectionSet.h
--- a/Engine/Source/Runtime/TypedElementRuntime/Public/Elements/Framework/TypedElementSelectionSet.h
+++ b/Engine/Source/Runtime/TypedElementRuntime/Public/Elements/Framework/TypedElementSelectionSet.h
@@ -250,6 +250,11 @@ public:
 	 */
 	TYPEDELEMENTRUNTIME_API bool AllowSelectionModifiers(const FTypedElementHandle& InElementHandle) const;
 
+	/**
+	 * Sets the name to use for teds integration so that it can be reapplied on undo/redo
+	 */
+	TYPEDELEMENTRUNTIME_API void SetNameForTedsIntegration(const FName InNameForIntegration);
+
 	/**
 	 * Given an element, return the element that should actually perform a selection operation.
 	 */
@@ -691,4 +696,6 @@ private:
 	 * and is then replaced with the state loaded from Serialize to be applied in PostEditUndo.
 	 */
 	TUniquePtr<FTypedElementSelectionSetState> PendingUndoRedoState;
+
+	FName ListNameForTedsIntegration;
 };
