[Typed Element Bridge] Add TEDS selection set name to the rest of the selection set options structs.

Intentionally not exposing these to BP - we don't really want that API in the wild

diff --git a/Engine/Source/Runtime/TypedElementRuntime/Public/Elements/Framework/TypedElementSelectionSet.h b/Engine/Source/Runtime/TypedElementRuntime/Public/Elements/Framework/TypedElementSelectionSet.h
--- a/Engine/Source/Runtime/TypedElementRuntime/Public/Elements/Framework/TypedElementSelectionSet.h
+++ b/Engine/Source/Runtime/TypedElementRuntime/Public/Elements/Framework/TypedElementSelectionSet.h
@@ -43,6 +43,10 @@ public:
 
 	FTypedElementSelectionNormalizationOptions& SetFollowAttachment(const bool InFollowAttachment) { bFollowAttachment = InFollowAttachment; return *this; }
 	bool FollowAttachment() const { return bFollowAttachment; }
+
+	// Set the selection set name that will be passed into the selection column in TEDS (if it is enabled)
+	FTypedElementSelectionNormalizationOptions& SetNameForTEDSIntegration(const FName& InTEDSIntegrationSelectionSetName) { TEDSIntegrationSelectionSetName = InTEDSIntegrationSelectionSetName; return *this; }
+	FName GetNameForTEDSIntegration() const { return TEDSIntegrationSelectionSetName; }
 	
 private:
 	UPROPERTY(BlueprintReadWrite, Category="TypedElementInterfaces|Selection|NormalizationOptions", meta=(AllowPrivateAccess=true))
@@ -50,6 +54,8 @@ private:
 
 	UPROPERTY(BlueprintReadWrite, Category="TypedElementInterfaces|Selection|NormalizationOptions", meta=(AllowPrivateAccess=true))
 	bool bFollowAttachment = false;
+
+	FName TEDSIntegrationSelectionSetName = FName();
 };
 
 /**
diff --git a/Engine/Source/Runtime/TypedElementRuntime/Public/Elements/Interfaces/TypedElementSelectionInterface.h b/Engine/Source/Runtime/TypedElementRuntime/Public/Elements/Interfaces/TypedElementSelectionInterface.h
--- a/Engine/Source/Runtime/TypedElementRuntime/Public/Elements/Interfaces/TypedElementSelectionInterface.h
+++ b/Engine/Source/Runtime/TypedElementRuntime/Public/Elements/Interfaces/TypedElementSelectionInterface.h
@@ -48,10 +48,16 @@ struct FTypedElementIsSelectedOptions
 public:
 	FTypedElementIsSelectedOptions& SetAllowIndirect(const bool InAllowIndirect) { bAllowIndirect = InAllowIndirect; return *this; }
 	bool AllowIndirect() const { return bAllowIndirect; }
+
+	// Set the selection set name that will be passed into the selection column in TEDS (if it is enabled)
+	FTypedElementIsSelectedOptions& SetNameForTEDSIntegration(const FName& InTEDSIntegrationSelectionSetName) { TEDSIntegrationSelectionSetName = InTEDSIntegrationSelectionSetName; return *this; }
+	FName GetNameForTEDSIntegration() const { return TEDSIntegrationSelectionSetName; }
 	
 private:
 	UPROPERTY(BlueprintReadWrite, Category="TypedElementInterfaces|Selection|IsSelectedOptions", meta=(AllowPrivateAccess=true))
 	bool bAllowIndirect = false;
+	
+	FName TEDSIntegrationSelectionSetName = FName();
 };
 
 USTRUCT(BlueprintType)
@@ -100,7 +106,6 @@ private:
 	UPROPERTY(BlueprintReadWrite, Category="TypedElementInterfaces|Selection|SelectionOptions", meta=(AllowPrivateAccess=true))
 	ETypedElementChildInclusionMethod ChildElementInclusionMethod = ETypedElementChildInclusionMethod::None;
 
-	UPROPERTY(BlueprintReadWrite, Category="TypedElementInterfaces|Selection|SelectionOptions", meta=(AllowPrivateAccess=true))
 	FName TEDSIntegrationSelectionSetName = FName();
 };
 
