[TypedElementRuntime] SelectionSet now sets the TEDS name option.

diff --git a/Engine/Source/Runtime/TypedElementRuntime/Private/Elements/Framework/TypedElementSelectionSet.cpp b/Engine/Source/Runtime/TypedElementRuntime/Private/Elements/Framework/TypedElementSelectionSet.cpp
--- a/Engine/Source/Runtime/TypedElementRuntime/Private/Elements/Framework/TypedElementSelectionSet.cpp
+++ b/Engine/Source/Runtime/TypedElementRuntime/Private/Elements/Framework/TypedElementSelectionSet.cpp
@@ -188,8 +188,13 @@ bool UTypedElementSelectionSet::CanDeselectElement(const FTypedElementHandle& In
 	return SelectionSetElement && SelectionSetElement.CanDeselectElement(InSelectionOptions);
 }
 
-bool UTypedElementSelectionSet::SelectElement(const FTypedElementHandle& InElementHandle, const FTypedElementSelectionOptions InSelectionOptions)
+bool UTypedElementSelectionSet::SelectElement(const FTypedElementHandle& InElementHandle, FTypedElementSelectionOptions InSelectionOptions)
 {
+	if (InSelectionOptions.GetNameForTEDSIntegration().IsNone())
+	{
+		InSelectionOptions.SetNameForTEDSIntegration(ListNameForTedsIntegration);
+	}
+
 	bool bSelectionChanged = false;
 
 	if (FTypedElementSelectionSetElement SelectionSetElement = ResolveSelectionSetElement(InElementHandle))
@@ -255,8 +260,13 @@ bool UTypedElementSelectionSet::SelectElements(FTypedElementListConstRef InEleme
 	return bSelectionChanged;
 }
 
-bool UTypedElementSelectionSet::DeselectElement(const FTypedElementHandle& InElementHandle, const FTypedElementSelectionOptions InSelectionOptions)
+bool UTypedElementSelectionSet::DeselectElement(const FTypedElementHandle& InElementHandle, FTypedElementSelectionOptions InSelectionOptions)
 {
+	if (InSelectionOptions.GetNameForTEDSIntegration().IsNone())
+	{
+		InSelectionOptions.SetNameForTEDSIntegration(ListNameForTedsIntegration);
+	}
+
 	bool bSelectionChanged = false;
 
 	if (FTypedElementSelectionSetElement SelectionSetElement = ResolveSelectionSetElement(InElementHandle))
diff --git a/Engine/Source/Runtime/TypedElementRuntime/Public/Elements/Framework/TypedElementSelectionSet.h b/Engine/Source/Runtime/TypedElementRuntime/Public/Elements/Framework/TypedElementSelectionSet.h
--- a/Engine/Source/Runtime/TypedElementRuntime/Public/Elements/Framework/TypedElementSelectionSet.h
+++ b/Engine/Source/Runtime/TypedElementRuntime/Public/Elements/Framework/TypedElementSelectionSet.h
@@ -211,7 +211,7 @@ public:
 	 * Attempt to select the given element.
 	 * @return True if the selection was changed, false otherwise.
 	 */
-	TYPEDELEMENTRUNTIME_API bool SelectElement(const FTypedElementHandle& InElementHandle, const FTypedElementSelectionOptions InSelectionOptions);
+	TYPEDELEMENTRUNTIME_API bool SelectElement(const FTypedElementHandle& InElementHandle, FTypedElementSelectionOptions InSelectionOptions);
 
 	/**
 	 * Attempt to select the given elements.
@@ -225,7 +225,7 @@ public:
 	 * Attempt to deselect the given element.
 	 * @return True if the selection was changed, false otherwise.
 	 */
-	TYPEDELEMENTRUNTIME_API bool DeselectElement(const FTypedElementHandle& InElementHandle, const FTypedElementSelectionOptions InSelectionOptions);
+	TYPEDELEMENTRUNTIME_API bool DeselectElement(const FTypedElementHandle& InElementHandle, FTypedElementSelectionOptions InSelectionOptions);
 
 	/**
 	 * Attempt to deselect the given elements.
