TypedElements: Give the element selection interface a chance to map to a different element before running customizations in UTypedElementSelectionSet.

#Synced-CL 49950069

diff --git a/Engine/Source/Runtime/TypedElementRuntime/Private/Elements/Framework/TypedElementSelectionSet.cpp b/Engine/Source/Runtime/TypedElementRuntime/Private/Elements/Framework/TypedElementSelectionSet.cpp
--- a/Engine/Source/Runtime/TypedElementRuntime/Private/Elements/Framework/TypedElementSelectionSet.cpp
+++ b/Engine/Source/Runtime/TypedElementRuntime/Private/Elements/Framework/TypedElementSelectionSet.cpp
@@ -5,6 +5,7 @@
 #include "Elements/Framework/TypedElementRegistry.h"
 #include "Elements/Framework/TypedElementUtil.h"
 #include "Elements/Interfaces/TypedElementHierarchyInterface.h"
+#include "HAL/IConsoleManager.h"
 
 #include "Misc/ITransaction.h"
 #include "Serialization/ObjectReader.h"
@@ -17,6 +18,17 @@
 #include "Serialization/TextReferenceCollector.h"
 #endif //WITH_EDITORONLY_DATA
 
+namespace TypedElementSelectionSetLocals
+{
+	bool bResolveElementGetSelectionElementFirst = true;
+	
+	FAutoConsoleVariableRef CVarResolveElementGetSelectionElementFirst(
+		TEXT("TypedElements.SelectionSet.ResolveElementGetSelectionElementFirst"),
+		bResolveElementGetSelectionElementFirst,
+		TEXT("If true UTypedElementSelectionSet::GetSelectionElement will allow the element "
+			"to provide a selection element before running any customizations."));
+}
+
 void FTypedElementSelectionCustomization::GetNormalizedElements(const TTypedElement<ITypedElementSelectionInterface>& InElementSelectionHandle, FTypedElementListConstRef InSelectionSet, const FTypedElementSelectionNormalizationOptions& InNormalizationOptions, FTypedElementListRef OutNormalizedElements)
 {
 	// Don't include this element in the normalized selection if it has a parent element that is also selected
@@ -410,7 +422,19 @@ FName UTypedElementSelectionSet::GetNameForTedsIntegration() const
 
 FTypedElementHandle UTypedElementSelectionSet::GetSelectionElement(const FTypedElementHandle& InElementHandle, const ETypedElementSelectionMethod InSelectionMethod) const
 {
-	FTypedElementSelectionSetElement SelectionSetElement = ResolveSelectionSetElement(InElementHandle);
+	FTypedElementHandle HandleToUse = InElementHandle;
+	
+	// Give the element interface a chance to remap to a different element first. This allows proxy elements to
+	//  map to their backing element and then have the backing element be used for any selection customization.
+	if (TypedElementSelectionSetLocals::bResolveElementGetSelectionElementFirst)
+	{
+		if (const TTypedElement<ITypedElementSelectionInterface> ElementSelectionInterface = GetElementList()->GetElement<ITypedElementSelectionInterface>(InElementHandle))
+		{
+			HandleToUse = ElementSelectionInterface.GetSelectionElement(GetElementList(), InSelectionMethod);
+		}
+	}
+
+	FTypedElementSelectionSetElement SelectionSetElement = ResolveSelectionSetElement(HandleToUse);
 	return SelectionSetElement ? SelectionSetElement.GetSelectionElement(InSelectionMethod) : FTypedElementHandle();
 }
 
